ScriptableTools: add support for setting custom icons for Tools

diff --git a/Engine/Plugins/Experimental/ScriptableToolsEditorMode/Source/ScriptableToolsEditorMode/Private/ScriptableToolsEditorModeStyle.cpp b/Engine/Plugins/Experimental/ScriptableToolsEditorMode/Source/ScriptableToolsEditorMode/Private/ScriptableToolsEditorModeStyle.cpp
--- a/Engine/Plugins/Experimental/ScriptableToolsEditorMode/Source/ScriptableToolsEditorMode/Private/ScriptableToolsEditorModeStyle.cpp
+++ b/Engine/Plugins/Experimental/ScriptableToolsEditorMode/Source/ScriptableToolsEditorMode/Private/ScriptableToolsEditorModeStyle.cpp
@@ -72,6 +72,55 @@ void FScriptableToolsEditorModeStyle::Initialize()
 	FSlateStyleRegistry::RegisterSlateStyle(*StyleSet.Get());
 };
 
+
+bool FScriptableToolsEditorModeStyle::TryRegisterCustomIcon(FName StyleIdentifier, FString FileNameWithPath, FString ExternalRelativePath)
+{
+	if ( ! ensure(StyleSet.IsValid()) )
+	{
+		return false;
+	}
+
+	const FVector2D Icon20x20(20.0f, 20.0f);
+
+	if (ExternalRelativePath.IsEmpty())
+	{
+		ExternalRelativePath = FileNameWithPath;
+	}
+
+	if (FPaths::FileExists(FileNameWithPath) == false)
+	{
+		UE_LOG(LogTemp, Warning, TEXT("Custom ScriptableTool Icon file not found at path %s (on-disk path is %s)"), *ExternalRelativePath, *FileNameWithPath);
+		return false;
+	}
+
+	FString Extension = FPaths::GetExtension(FileNameWithPath);
+	if ( !(Extension.Equals("svg", ESearchCase::IgnoreCase) || Extension.Equals("png", ESearchCase::IgnoreCase)) )
+	{
+		UE_LOG(LogTemp, Warning, TEXT("Custom ScriptableTool Icon at path %s has unsupported type, must be svg or png"), *ExternalRelativePath);
+		return false;
+	}
+
+	// need to re-register style to be able to modify it
+	FSlateStyleRegistry::UnRegisterSlateStyle(*StyleSet.Get());
+
+	if (Extension.Equals("svg", ESearchCase::IgnoreCase))
+	{
+		StyleSet->Set(StyleIdentifier,
+			new FSlateVectorImageBrush(FileNameWithPath, Icon20x20));
+	}
+	else if (Extension.Equals("png", ESearchCase::IgnoreCase))
+	{
+		StyleSet->Set(StyleIdentifier,
+			new FSlateImageBrush(FileNameWithPath, Icon20x20) );
+	}
+
+	FSlateStyleRegistry::RegisterSlateStyle(*StyleSet.Get());
+
+	return true;
+}
+
+
+
 END_SLATE_FUNCTION_BUILD_OPTIMIZATION
 
 #undef IMAGE_PLUGIN_BRUSH
diff --git a/Engine/Plugins/Experimental/ScriptableToolsEditorMode/Source/ScriptableToolsEditorMode/Private/ScriptableToolsEditorModeToolkit.cpp b/Engine/Plugins/Experimental/ScriptableToolsEditorMode/Source/ScriptableToolsEditorMode/Private/ScriptableToolsEditorModeToolkit.cpp
--- a/Engine/Plugins/Experimental/ScriptableToolsEditorMode/Source/ScriptableToolsEditorMode/Private/ScriptableToolsEditorModeToolkit.cpp
+++ b/Engine/Plugins/Experimental/ScriptableToolsEditorMode/Source/ScriptableToolsEditorMode/Private/ScriptableToolsEditorModeToolkit.cpp
@@ -9,6 +9,7 @@
 #include "Framework/MultiBox/MultiBoxBuilder.h"
 #include "IDetailsView.h"
 #include "Toolkits/AssetEditorModeUILayer.h"
+#include "Interfaces/IPluginManager.h"
 
 #include "InteractiveToolManager.h"
 
@@ -391,9 +392,51 @@ void FScriptableToolsEditorModeToolkit::BuildToolPalette(FName PaletteIndex, cla
 		FName InExtensionHook = NAME_None;
 		TAttribute<FText> Label = ToolCDO->ToolName.IsEmpty() ? LOCTEXT("EmptyToolName", "Tool") : ToolCDO->ToolName;
 		TAttribute<FText> Tooltip = ToolCDO->ToolTooltip.IsEmpty() ? FText() : ToolCDO->ToolTooltip;
-		// hardcoded icon name for now
+		
+		// default icon comes with the mode
 		TAttribute<FSlateIcon> Icon = FSlateIcon(FScriptableToolsEditorModeStyle::Get()->GetStyleSetName(), "ScriptableToolsEditorModeToolCommands.DefaultToolIcon");
-			
+
+		// if a custom icon is defined, try to find it, this can fail in many ways, in that case
+		// the default icon is kept
+		if (ToolCDO->CustomIconPath.IsEmpty() == false)
+		{
+			FName ToolIconToken( FString("ScriptableToolsEditorModeToolCommands.") + ToolIdentifier );
+
+			// Custom Tool Icons are assumed to reside in the same Content folder as the Plugin/Project that
+			// the Tool Class is defined in, and that the CustomIconPath is a relative path inside that Content folder.
+			// use the class Package path to determine if this it is in a Plugin or directly in the Project, so that
+			// we can get the right ContentDir below.
+			// (Note that a relative ../../../ style path can always be used to redirect to any other file...)
+			FString FullPathName = ToolCDO->GetClass()->GetPathName();
+			FString PathPart, FilenamePart, ExtensionPart;
+			FPaths::Split(FullPathName, PathPart, FilenamePart, ExtensionPart);
+
+			FString FullIconPath = ToolCDO->CustomIconPath;
+			if (PathPart.StartsWith("/Game"))
+			{
+				FullIconPath = FPaths::ProjectContentDir() / ToolCDO->CustomIconPath;
+			}
+			else
+			{
+				TArray<FString> PathParts;
+				PathPart.ParseIntoArray(PathParts, TEXT("/"));
+				if (PathParts.Num() > 0)
+				{
+					FString PluginContentDir = IPluginManager::Get().FindPlugin(PathParts[0])->GetContentDir();
+					FullIconPath = PluginContentDir / ToolCDO->CustomIconPath;
+				}
+				else  // something is wrong, fall back to project content dir
+				{
+					FullIconPath = FPaths::ProjectContentDir() / ToolCDO->CustomIconPath;
+				}
+			}
+
+			if (FScriptableToolsEditorModeStyle::TryRegisterCustomIcon(ToolIconToken, FullIconPath, ToolCDO->CustomIconPath))
+			{
+				Icon = FSlateIcon(FScriptableToolsEditorModeStyle::Get()->GetStyleSetName(), ToolIconToken);
+			}
+		}
+		
 		ToolbarBuilder.AddToolBarButton(*NewAction, InExtensionHook, Label, Tooltip, Icon);
 
 	});
diff --git a/Engine/Plugins/Experimental/ScriptableToolsEditorMode/Source/ScriptableToolsEditorMode/Public/ScriptableToolsEditorModeStyle.h b/Engine/Plugins/Experimental/ScriptableToolsEditorMode/Source/ScriptableToolsEditorMode/Public/ScriptableToolsEditorModeStyle.h
--- a/Engine/Plugins/Experimental/ScriptableToolsEditorMode/Source/ScriptableToolsEditorMode/Public/ScriptableToolsEditorModeStyle.h
+++ b/Engine/Plugins/Experimental/ScriptableToolsEditorMode/Source/ScriptableToolsEditorMode/Public/ScriptableToolsEditorModeStyle.h
@@ -20,6 +20,20 @@ public:
 	// use to access icons defined by the style set by name, eg GetBrush("BrushFalloffIcons.Smooth")
 	static const FSlateBrush* GetBrush(FName PropertyName, const ANSICHAR* Specifier = NULL);
 
+
+	// Try to find a png or svg icon at FileNameWithPath and register it in the style set with the given StyleIconIdentifier,
+	// ie so that later it can be found FSlateIcon(FScriptableToolsEditorModeStyle::Get()->GetStyleSetName(), StyleIconIdentifier).
+	// This is intended to be used for (eg) custom icons defined in Scriptable Tool Blueprints, that are not known at compile time.
+	// 
+	// Will print error and return false if image is not found, or not png/svg format.
+	// 
+	// FileNameWithPath must be the full on-disk path to the icon, which is constructed in code from a relative path
+	// the user enters as (eg) a property/etc. To help with debugging, this path can be provided as ExternalRelativePath,
+	// and it will be printed out instead of the absolute path as part of the log error message.
+	//
+	static bool TryRegisterCustomIcon(FName StyleIconIdentifier, FString FileNameWithPath, FString ExternalRelativePath);
+
+
 private:
 	static FString InContent(const FString& RelativePath, const ANSICHAR* Extension);
 
diff --git a/Engine/Plugins/Experimental/ScriptableToolsFramework/Source/ScriptableToolsFramework/Public/ScriptableInteractiveTool.h b/Engine/Plugins/Experimental/ScriptableToolsFramework/Source/ScriptableToolsFramework/Public/ScriptableInteractiveTool.h
--- a/Engine/Plugins/Experimental/ScriptableToolsFramework/Source/ScriptableToolsFramework/Public/ScriptableInteractiveTool.h
+++ b/Engine/Plugins/Experimental/ScriptableToolsFramework/Source/ScriptableToolsFramework/Public/ScriptableInteractiveTool.h
@@ -288,6 +288,19 @@ public:
 	UPROPERTY(BlueprintReadWrite, EditAnywhere, Category = "Scriptable Tool Settings", meta=(DisplayName="Tooltip"))
 	FText ToolTooltip;
 
+	/**
+	 * Relative Path to a custom Icon Image for this Tool. The Image file format must be png or svg. 
+	 *
+	 * The Image file must reside in the same Content folder hierarchy as contains the Tool Class (ie Blueprint asset), 
+	 * i.e. either the Project Content folder or a Plugin Content folder.
+	 * 
+	 * So for example if the Tool BP is in a plugin named MyCustomTools, the icon must be in 
+	 *      MyProject/Plugins/MyCustomTools/Content/<SubFolders>/MyToolIcon.png,
+	 * and the relative path to use here would be <SubFolders>/MyToolIcon.png
+	 */ 
+	UPROPERTY(BlueprintReadWrite, EditAnywhere, Category = "Scriptable Tool Settings", meta = (DisplayName = "Custom Icon Path"))
+	FString CustomIconPath;
+
 	/** A generic flag to indicate whether this Tool should be shown in the UE Editor. This may be interpreted in different ways */
 	UPROPERTY(BlueprintReadWrite, EditAnywhere, Category = "Scriptable Tool Settings", meta=(DisplayName="Visible in Editor"))
 	bool bShowToolInEditor = true;
