ScriptableTools: Turn off EnabledByDefault.

diff --git a/Engine/Plugins/Editor/ScriptableToolsEditorMode/ScriptableToolsEditorMode.uplugin b/Engine/Plugins/Editor/ScriptableToolsEditorMode/ScriptableToolsEditorMode.uplugin
--- a/Engine/Plugins/Editor/ScriptableToolsEditorMode/ScriptableToolsEditorMode.uplugin
+++ b/Engine/Plugins/Editor/ScriptableToolsEditorMode/ScriptableToolsEditorMode.uplugin
@@ -10,7 +10,7 @@
 	"DocsURL": "",
 	"MarketplaceURL": "",
 	"SupportURL": "",
-    "EnabledByDefault": true,
+    "EnabledByDefault": false,
 	"CanContainContent": true,
 	"IsBetaVersion": true,
 	"Installed": false,
