ScriptableTools: Added support for a UMG defined overlay configured via a ScriptableTool  blueprint, allowing for customized in viewport UX

diff --git a/Engine/Plugins/Editor/ScriptableToolsEditorMode/Source/ScriptableToolsEditorMode/Private/ScriptableToolsEditorMode.cpp b/Engine/Plugins/Editor/ScriptableToolsEditorMode/Source/ScriptableToolsEditorMode/Private/ScriptableToolsEditorMode.cpp
--- a/Engine/Plugins/Editor/ScriptableToolsEditorMode/Source/ScriptableToolsEditorMode/Private/ScriptableToolsEditorMode.cpp
+++ b/Engine/Plugins/Editor/ScriptableToolsEditorMode/Source/ScriptableToolsEditorMode/Private/ScriptableToolsEditorMode.cpp
@@ -28,6 +28,9 @@
 #include "ToolTargets/DynamicMeshComponentToolTarget.h"
 #include "ToolTargets/SkeletalMeshComponentToolTarget.h"
 
+#include "Utility/ScriptableToolContextObjects.h"
+#include "ContextObjectStore.h"
+
 #define LOCTEXT_NAMESPACE "UScriptableToolsEditorMode"
 
 const FEditorModeID UScriptableToolsEditorMode::EM_ScriptableToolsEditorModeId = TEXT("EM_ScriptableToolsEditorMode");
@@ -170,7 +173,8 @@ void UScriptableToolsEditorMode::Enter()
 	// register each of them with ToolManager
 	ScriptableTools->ForEachScriptableTool([&](UClass* ToolClass, UInteractiveToolBuilder* ToolBuilder) 
 	{
-		FString UseName = ToolClass->GetName();
+		FString UseName;
+		ToolClass->GetClassPathName().ToString(UseName);
 		GetToolManager(EToolsContextScope::EdMode)->RegisterToolType(UseName, ToolBuilder);
 	});
 
@@ -187,8 +191,38 @@ void UScriptableToolsEditorMode::Enter()
 		ModeToolkit->InitializeAfterModeSetup();
 		ModeToolkit->ForceToolPaletteRebuild();
 	}
+
+	InitializeModeContexts();
 }
 
+void UScriptableToolsEditorMode::InitializeModeContexts()
+{
+	UContextObjectStore* ContextStore = GetInteractiveToolsContext()->ToolManager->GetContextObjectStore();
+
+	auto AddContextObject = [this, ContextStore](UScriptableToolContextObject* Object)
+	{
+		if (ensure(ContextStore->AddContextObject(Object)))
+		{
+			ContextsToShutdown.Add(Object);
+		}
+		ContextsToUpdateOnToolEnd.Add(Object);
+	};
+
+	UScriptableToolViewportWidgetAPI* ViewportWidgetAPI = NewObject<UScriptableToolViewportWidgetAPI>();
+	ViewportWidgetAPI = NewObject<UScriptableToolViewportWidgetAPI>();
+	ViewportWidgetAPI->Initialize(
+		[this](TSharedRef<SWidget> InOverlaidWidget) {
+			Toolkit->GetToolkitHost()->AddViewportOverlayWidget(InOverlaidWidget);
+		},
+		[this](TSharedRef<SWidget> InOverlaidWidget) {
+			Toolkit->GetToolkitHost()->RemoveViewportOverlayWidget(InOverlaidWidget);
+		}
+		);
+	AddContextObject(ViewportWidgetAPI);
+
+}
+
+
 
 void UScriptableToolsEditorMode::OnBlueprintPreCompile(UBlueprint* Blueprint)
 {
@@ -234,6 +268,16 @@ void UScriptableToolsEditorMode::Exit()
 	// clear realtime viewport override
 	ConfigureRealTimeViewportsOverride(false);
 
+	UContextObjectStore* ContextStore = GetInteractiveToolsContext()->ToolManager->GetContextObjectStore();
+	for (TWeakObjectPtr<UScriptableToolContextObject> Context : ContextsToShutdown)
+	{
+		if (Context.IsValid())
+		{
+			Context->Shutdown();
+			ContextStore->RemoveContextObject(Context.Get());
+		}
+	}
+
 	// Call base Exit method to ensure proper cleanup
 	UEdMode::Exit();
 }
@@ -284,6 +328,14 @@ void UScriptableToolsEditorMode::OnToolEnded(UInteractiveToolManager* Manager, U
 {
 	// re-enable slate throttling (see OnToolStarted)
 	FSlateThrottleManager::Get().DisableThrottle(false);
+
+	for (TWeakObjectPtr<UScriptableToolContextObject> Context : ContextsToUpdateOnToolEnd)
+	{
+		if (Context.IsValid())
+		{
+			Context->OnToolEnded(Tool);
+		}
+	}
 }
 
 void UScriptableToolsEditorMode::BindCommands()
diff --git a/Engine/Plugins/Editor/ScriptableToolsEditorMode/Source/ScriptableToolsEditorMode/Public/ScriptableToolsEditorMode.h b/Engine/Plugins/Editor/ScriptableToolsEditorMode/Source/ScriptableToolsEditorMode/Public/ScriptableToolsEditorMode.h
--- a/Engine/Plugins/Editor/ScriptableToolsEditorMode/Source/ScriptableToolsEditorMode/Public/ScriptableToolsEditorMode.h
+++ b/Engine/Plugins/Editor/ScriptableToolsEditorMode/Source/ScriptableToolsEditorMode/Public/ScriptableToolsEditorMode.h
@@ -18,6 +18,7 @@ class UModelingSelectionInteraction;
 class UGeometrySelectionManager;
 class UScriptableToolSet;
 class UBlueprint;
+class UScriptableToolContextObject;
 
 UCLASS(Transient)
 class SCRIPTABLETOOLSEDITORMODE_API UScriptableToolsEditorMode : public UBaseLegacyWidgetEdMode
@@ -63,6 +64,7 @@ public:
 	// End of UEdMode interface
 	//////////////////
 
+	
 
 protected:
 	virtual void BindCommands() override;
@@ -87,6 +89,13 @@ protected:
 	void OnBlueprintCompiled();
 	FDelegateHandle BlueprintCompiledHandle;
 
+	void InitializeModeContexts();
+
+private:
+
+	TArray<TWeakObjectPtr<UScriptableToolContextObject>> ContextsToUpdateOnToolEnd;
+	TArray<TWeakObjectPtr<UScriptableToolContextObject>> ContextsToShutdown;
+
 protected:
 	UPROPERTY()
 	TObjectPtr<UScriptableToolSet> ScriptableTools;
diff --git a/Engine/Plugins/Runtime/ScriptableToolsFramework/ScriptableToolsFramework.uplugin b/Engine/Plugins/Runtime/ScriptableToolsFramework/ScriptableToolsFramework.uplugin
--- a/Engine/Plugins/Runtime/ScriptableToolsFramework/ScriptableToolsFramework.uplugin
+++ b/Engine/Plugins/Runtime/ScriptableToolsFramework/ScriptableToolsFramework.uplugin
@@ -30,6 +30,10 @@
 		{
 			"Name": "MeshModelingToolset",
 			"Enabled": true
+		},
+		{
+			"Name": "MeshModelingToolsetExp",
+			"Enabled": true
 		}
 	]
 }
diff --git a/Engine/Plugins/Runtime/ScriptableToolsFramework/Source/ScriptableToolsFramework/Private/ScriptableInteractiveTool.cpp b/Engine/Plugins/Runtime/ScriptableToolsFramework/Source/ScriptableToolsFramework/Private/ScriptableInteractiveTool.cpp
--- a/Engine/Plugins/Runtime/ScriptableToolsFramework/Source/ScriptableToolsFramework/Private/ScriptableInteractiveTool.cpp
+++ b/Engine/Plugins/Runtime/ScriptableToolsFramework/Source/ScriptableToolsFramework/Private/ScriptableInteractiveTool.cpp
@@ -21,10 +21,23 @@
 #include "InteractiveGizmoManager.h"
 #include "SceneView.h"
 
+#include "Blueprint/UserWidget.h"
+#include "ContextObjectStore.h"
+#include "Utility/ScriptableToolContextObjects.h"
+#include "Slate/SObjectWidget.h"
+#include "ModelingWidgets/SDraggableBox.h"
+#include "Components/Widget.h"
+
+
 #include "UObject/EnumProperty.h"
 
 #define LOCTEXT_NAMESPACE "UScriptableInteractiveTool"
 
+namespace ScriptableInteractiveToolLocals
+{
+	class SScriptableToolViewportOverlayDragContainer : public SObjectWidget
+	{ };
+}
 
 UScriptableInteractiveTool* UScriptableInteractiveToolPropertySet::GetOwningTool(EToolsFrameworkOutcomePins& Outcome)
 {
@@ -194,6 +207,7 @@ void UScriptableInteractiveTool::Setup()
 	OnScriptSetup();
 
 
+
 	ToolDrawableGeometry = NewObject<UPreviewGeometry>();
 	ToolDrawableGeometry->CreateInWorld(GetWorld(), FTransform::Identity);
 
@@ -1224,6 +1238,46 @@ UScriptableToolTriangleSet* UScriptableInteractiveTool::AddTriangleSet()
 }
 
 
+void UScriptableInteractiveTool::SetOverlayWidget(UUserWidget* Widget, bool bMakeDraggable )
+{
+	UContextObjectStore* ContextStore = GetToolManager()->GetContextObjectStore();
+
+	UScriptableToolViewportWidgetAPI* ViewportWidgetAPI = ContextStore->FindContext< UScriptableToolViewportWidgetAPI>();
+	if (ViewportWidgetAPI && Widget)
+	{
+		if (bMakeDraggable)
+		{
+
+			TSharedRef<SWidget> WrappedWidget = Widget->TakeDerivedWidget<ScriptableInteractiveToolLocals::SScriptableToolViewportOverlayDragContainer>([](UUserWidget* Widget, TSharedRef<SWidget> Content) {
+				return SNew(ScriptableInteractiveToolLocals::SScriptableToolViewportOverlayDragContainer, Widget)
+					[
+						SNew(SDraggableBoxOverlay)
+						[
+							Content
+						]
+					];
+				});
+
+			ViewportWidgetAPI->SetOverlayWidget(WrappedWidget);
+
+		}
+		else
+		{
+			ViewportWidgetAPI->SetOverlayWidget(Widget->TakeWidget());
+		}
+	}
+}
+
+void UScriptableInteractiveTool::ClearOverlayWidget()
+{
+	UContextObjectStore* ContextStore = GetToolManager()->GetContextObjectStore();
+
+	UScriptableToolViewportWidgetAPI* ViewportWidgetAPI = ContextStore->FindContext< UScriptableToolViewportWidgetAPI>();
+	if (ViewportWidgetAPI)
+	{
+		ViewportWidgetAPI->ClearOverlayWidget();
+	}
+}
 
 FInputRayHit UScriptableToolsUtilityLibrary::MakeInputRayHit_Miss()
 {
diff --git a/Engine/Plugins/Runtime/ScriptableToolsFramework/Source/ScriptableToolsFramework/Public/ScriptableInteractiveTool.h b/Engine/Plugins/Runtime/ScriptableToolsFramework/Source/ScriptableToolsFramework/Public/ScriptableInteractiveTool.h
--- a/Engine/Plugins/Runtime/ScriptableToolsFramework/Source/ScriptableToolsFramework/Public/ScriptableInteractiveTool.h
+++ b/Engine/Plugins/Runtime/ScriptableToolsFramework/Source/ScriptableToolsFramework/Public/ScriptableInteractiveTool.h
@@ -18,10 +18,11 @@ class FCanvas;
 class FSceneView;
 
 class UPreviewGeometry;
+
 class UScriptableToolLineSet;
 class UScriptableToolPointSet;
 class UScriptableToolTriangleSet;
-
+class UUserWidget;
 
 UENUM(BlueprintType)
 enum class EToolsFrameworkOutcomePins : uint8
@@ -880,6 +881,17 @@ protected:
 	UPROPERTY(Transient, DuplicateTransient, NonTransactional, SkipSerialization)
 	TObjectPtr<UPreviewGeometry> ToolDrawableGeometry = nullptr;
 
+
+public:
+
+	// Widget Overlay API
+
+	UFUNCTION(BlueprintCallable, Category = "ScriptableTool|Widgets")
+	void SetOverlayWidget(UUserWidget* Widget, bool bMakeDraggable=true);
+
+	UFUNCTION(BlueprintCallable, Category = "ScriptableTool|Widgets")
+	void ClearOverlayWidget();
+
 };
 
 
diff --git a/Engine/Plugins/Runtime/ScriptableToolsFramework/Source/ScriptableToolsFramework/Public/Utility/ScriptableToolContextObjects.h b/Engine/Plugins/Runtime/ScriptableToolsFramework/Source/ScriptableToolsFramework/Public/Utility/ScriptableToolContextObjects.h
new file mode 100644
--- /dev/null
+++ b/Engine/Plugins/Runtime/ScriptableToolsFramework/Source/ScriptableToolsFramework/Public/Utility/ScriptableToolContextObjects.h
@@ -0,0 +1,85 @@
+// Copyright Epic Games, Inc. All Rights Reserved.
+
+#pragma once
+
+#include "CoreMinimal.h"
+
+#include "InteractiveToolManager.h"
+#include "Widgets/SWidget.h"
+
+#include "ScriptableToolContextObjects.generated.h"
+
+/**
+ * Base class for context objects used in the Scriptable Tools Framework.
+ */
+UCLASS()
+class SCRIPTABLETOOLSFRAMEWORK_API UScriptableToolContextObject : public UObject
+{
+	GENERATED_BODY()
+public:
+	/**
+	 * Called by the mode when shutting context objects down, allowing them to do any cleanup.
+	 * Initialization, on the other hand is usually done by some class-specific Initialize() method.
+	 */
+	virtual void Shutdown() {}
+
+	/**
+	 * Called whenever a tool is ended, for instance to let a context object remove listeners associated
+	 * with that tool (it shouldn't have to do so, but may choose to for robustness).
+	 */
+	virtual void OnToolEnded(UInteractiveTool* DeadTool) {}
+};
+
+
+UCLASS()
+class SCRIPTABLETOOLSFRAMEWORK_API UScriptableToolViewportWidgetAPI : public UScriptableToolContextObject
+{
+	GENERATED_BODY()
+
+public:
+
+	void Initialize(TUniqueFunction<void(TSharedRef<SWidget> InOverlaidWidget)> ReplaceOverlaidWidgetFuncIn,
+	               	TUniqueFunction<void(TSharedRef<SWidget> InOverlaidWidget)> ClearOverlaidWidgetFuncIn)
+	{
+		ReplaceOverlaidWidgetFunc = MoveTemp(ReplaceOverlaidWidgetFuncIn);
+		ClearOverlaidWidgetFunc = MoveTemp(ClearOverlaidWidgetFuncIn);
+	}
+
+	void SetOverlayWidget(TSharedRef<SWidget> InOverlaidWidget)
+	{
+		if (CurrentOverlaidWidget.IsValid())
+		{
+			ClearOverlayWidget();
+		}
+
+		CurrentOverlaidWidget = InOverlaidWidget.ToWeakPtr();
+
+		ReplaceOverlaidWidgetFunc(InOverlaidWidget);
+	}
+	void ClearOverlayWidget()
+	{
+		if (CurrentOverlaidWidget.IsValid())
+		{
+			ClearOverlaidWidgetFunc(CurrentOverlaidWidget.Pin().ToSharedRef());
+			CurrentOverlaidWidget.Reset();
+		}
+	}
+
+	virtual void Shutdown() override 
+	{
+		ClearOverlayWidget();
+	}
+
+
+	virtual void OnToolEnded(UInteractiveTool* DeadTool) override
+	{
+		ClearOverlayWidget();
+	}
+
+private:
+
+	TUniqueFunction<void(TSharedRef<SWidget> InOverlaidWidget)> ReplaceOverlaidWidgetFunc;
+	TUniqueFunction<void(TSharedRef<SWidget> InOverlaidWidget)> ClearOverlaidWidgetFunc;
+
+	TWeakPtr<SWidget> CurrentOverlaidWidget;
+};
\ No newline at end of file
diff --git a/Engine/Plugins/Runtime/ScriptableToolsFramework/Source/ScriptableToolsFramework/ScriptableToolsFramework.Build.cs b/Engine/Plugins/Runtime/ScriptableToolsFramework/Source/ScriptableToolsFramework/ScriptableToolsFramework.Build.cs
--- a/Engine/Plugins/Runtime/ScriptableToolsFramework/Source/ScriptableToolsFramework/ScriptableToolsFramework.Build.cs
+++ b/Engine/Plugins/Runtime/ScriptableToolsFramework/Source/ScriptableToolsFramework/ScriptableToolsFramework.Build.cs
@@ -18,6 +18,7 @@ public class ScriptableToolsFramework : ModuleRules
 		PrivateIncludePaths.AddRange(
 			new string[] {
 				// ... add other private include paths required here ...
+
 			}
 			);
 
@@ -32,7 +33,9 @@ public class ScriptableToolsFramework : ModuleRules
 				"GeometryCore",
 				"InputCore",
 				"InteractiveToolsFramework",
-				"ModelingComponents"
+				"ModelingComponents",
+				"UMG",
+				"ModelingUI"
 			}
 			);
 
@@ -40,7 +43,9 @@ public class ScriptableToolsFramework : ModuleRules
 		PrivateDependencyModuleNames.AddRange(
 			new string[]
 			{
-				"Engine"
+				"Engine",
+				"Slate",
+				"SlateCore"
 			}
 			);
 
