ScriptableTools: Fix naming issue that prevents ScriptableTools from being launched.

diff --git a/Engine/Plugins/Editor/ScriptableToolsEditorMode/Source/ScriptableToolsEditorMode/Private/ScriptableToolsEditorModeToolkit.cpp b/Engine/Plugins/Editor/ScriptableToolsEditorMode/Source/ScriptableToolsEditorMode/Private/ScriptableToolsEditorModeToolkit.cpp
--- a/Engine/Plugins/Editor/ScriptableToolsEditorMode/Source/ScriptableToolsEditorMode/Private/ScriptableToolsEditorModeToolkit.cpp
+++ b/Engine/Plugins/Editor/ScriptableToolsEditorMode/Source/ScriptableToolsEditorMode/Private/ScriptableToolsEditorModeToolkit.cpp
@@ -365,7 +365,9 @@ void FScriptableToolsEditorModeToolkit::BuildToolPalette(FName PaletteIndex, cla
 			return;
 		}
 
-		FString ToolIdentifier = ToolClass->GetName();
+		
+		FString ToolIdentifier;
+		ToolClass->GetClassPathName().ToString(ToolIdentifier);
 
 		TSharedPtr<FUIAction> NewAction = MakeShared<FUIAction>(
 		FExecuteAction::CreateLambda([this, ToolClass, ToolIdentifier, ToolManager]()
