ScriptableTools: Added support for tool registration filtering, provided by group tags attached to tools and listed in Editor Settings.

Tags are defined as blueprint objects to be included along with collections of tools. Tools can use more than one tag and any allowed tag in settings will enable the tool's registration upon mode startup.

diff --git a/Engine/Plugins/Editor/ScriptableToolsEditorMode/Source/ScriptableToolsEditorMode/Private/ScriptableToolsEditorMode.cpp b/Engine/Plugins/Editor/ScriptableToolsEditorMode/Source/ScriptableToolsEditorMode/Private/ScriptableToolsEditorMode.cpp
--- a/Engine/Plugins/Editor/ScriptableToolsEditorMode/Source/ScriptableToolsEditorMode/Private/ScriptableToolsEditorMode.cpp
+++ b/Engine/Plugins/Editor/ScriptableToolsEditorMode/Source/ScriptableToolsEditorMode/Private/ScriptableToolsEditorMode.cpp
@@ -12,6 +12,7 @@
 #include "InteractiveToolManager.h"
 #include "ScriptableToolsEditorModeToolkit.h"
 #include "ScriptableToolsEditorModeManagerCommands.h"
+#include "ScriptableToolsEditorModeSettings.h"
 
 #include "BaseGizmos/TransformGizmoUtil.h"
 #include "Snapping/ModelingSceneSnappingManager.h"
@@ -31,6 +32,10 @@
 #include "Utility/ScriptableToolContextObjects.h"
 #include "ContextObjectStore.h"
 
+#include "Engine/StreamableManager.h"
+
+
+
 #define LOCTEXT_NAMESPACE "UScriptableToolsEditorMode"
 
 const FEditorModeID UScriptableToolsEditorMode::EM_ScriptableToolsEditorModeId = TEXT("EM_ScriptableToolsEditorMode");
@@ -166,16 +171,12 @@ void UScriptableToolsEditorMode::Enter()
 	// enable realtime viewport override
 	ConfigureRealTimeViewportsOverride(true);
 
-
 	ScriptableTools = NewObject<UScriptableToolSet>(this);
-	// find all the Tool Blueprints
-	ScriptableTools->ReinitializeScriptableTools();
-	// register each of them with ToolManager
-	ScriptableTools->ForEachScriptableTool([&](UClass* ToolClass, UInteractiveToolBuilder* ToolBuilder) 
+
+	UScriptableToolsModeCustomizationSettings* ModeSettings = GetMutableDefault<UScriptableToolsModeCustomizationSettings>();
+	ModeSettings->OnSettingChanged().AddLambda([this](UObject*, FPropertyChangedEvent&)
 	{
-		FString UseName;
-		ToolClass->GetClassPathName().ToString(UseName);
-		GetToolManager(EToolsContextScope::EdMode)->RegisterToolType(UseName, ToolBuilder);
+		RebuildScriptableToolSet();
 	});
 
 	// todoz
@@ -189,12 +190,80 @@ void UScriptableToolsEditorMode::Enter()
 	{
 		FScriptableToolsEditorModeToolkit* ModeToolkit = (FScriptableToolsEditorModeToolkit*)Toolkit.Get();
 		ModeToolkit->InitializeAfterModeSetup();
-		ModeToolkit->ForceToolPaletteRebuild();
 	}
 
+	RebuildScriptableToolSet();
+
 	InitializeModeContexts();
 }
 
+void UScriptableToolsEditorMode::RebuildScriptableToolSet()
+{
+	UScriptableToolsModeCustomizationSettings* ModeSettings = GetMutableDefault<UScriptableToolsModeCustomizationSettings>();
+
+	auto UnregisterTools = [this]()
+	{
+		// unregister old tools from ToolManager
+		ScriptableTools->ForEachScriptableTool([&](UClass* ToolClass, UInteractiveToolBuilder* ToolBuilder)
+			{
+				FString UseName;
+				ToolClass->GetClassPathName().ToString(UseName);
+				GetToolManager(EToolsContextScope::EdMode)->UnregisterToolType(UseName);
+			});
+
+		if (Toolkit.IsValid())
+		{
+			FScriptableToolsEditorModeToolkit* ModeToolkit = (FScriptableToolsEditorModeToolkit*)Toolkit.Get();
+			ModeToolkit->StartAsyncToolLoading();
+		};
+	};
+
+	auto RegisterTools = [this]()
+	{
+		// register each of them with ToolManager
+		ScriptableTools->ForEachScriptableTool([&](UClass* ToolClass, UInteractiveToolBuilder* ToolBuilder)
+		{
+			FString UseName;
+			ToolClass->GetClassPathName().ToString(UseName);
+			GetToolManager(EToolsContextScope::EdMode)->RegisterToolType(UseName, ToolBuilder);
+		});
+
+		if (Toolkit.IsValid())
+		{
+			FScriptableToolsEditorModeToolkit* ModeToolkit = (FScriptableToolsEditorModeToolkit*)Toolkit.Get();
+			ModeToolkit->EndAsyncToolLoading();
+			ModeToolkit->ForceToolPaletteRebuild();
+		}
+	};
+
+	auto ToolLoadingUpdate = [this](TSharedPtr<FStreamableHandle> Handle)
+	{
+		if (Toolkit.IsValid())
+		{
+			FScriptableToolsEditorModeToolkit* ModeToolkit = (FScriptableToolsEditorModeToolkit*)Toolkit.Get();
+			ModeToolkit->SetAsyncProgress(Handle->GetProgress());
+		}
+	};
+
+	// find all the Tool Blueprints
+	if (ModeSettings->RegisterAllTools())
+	{
+		ScriptableTools->ReinitializeScriptableTools(FToolsLoadedDelegate::CreateLambda(UnregisterTools),
+													 FToolsLoadedDelegate::CreateLambda(RegisterTools),
+			                                         FToolsLoadingUpdateDelegate::CreateLambda(ToolLoadingUpdate));
+	}
+	else
+	{
+		ScriptableTools->ReinitializeScriptableTools(FToolsLoadedDelegate::CreateLambda(UnregisterTools),
+												     FToolsLoadedDelegate::CreateLambda(RegisterTools),
+			                                         FToolsLoadingUpdateDelegate::CreateLambda(ToolLoadingUpdate),
+			                                         &ModeSettings->ToolRegistrationFilters);
+	}	
+
+
+
+}
+
 void UScriptableToolsEditorMode::InitializeModeContexts()
 {
 	UContextObjectStore* ContextStore = GetInteractiveToolsContext()->ToolManager->GetContextObjectStore();
@@ -236,13 +305,7 @@ void UScriptableToolsEditorMode::OnBlueprintPreCompile(UBlueprint* Blueprint)
 
 void UScriptableToolsEditorMode::OnBlueprintCompiled()
 {
-	// Probably not necessary to always rebuild the palette here. But currently this lets us respond
-	// to changes in the tool name/setting/etc
-	if (Toolkit.IsValid())
-	{
-		FScriptableToolsEditorModeToolkit* ModeToolkit = (FScriptableToolsEditorModeToolkit*)Toolkit.Get();
-		ModeToolkit->ForceToolPaletteRebuild();
-	}
+	RebuildScriptableToolSet();
 }
 
 void UScriptableToolsEditorMode::Exit()
diff --git a/Engine/Plugins/Editor/ScriptableToolsEditorMode/Source/ScriptableToolsEditorMode/Private/ScriptableToolsEditorModeModule.cpp b/Engine/Plugins/Editor/ScriptableToolsEditorMode/Source/ScriptableToolsEditorMode/Private/ScriptableToolsEditorModeModule.cpp
--- a/Engine/Plugins/Editor/ScriptableToolsEditorMode/Source/ScriptableToolsEditorMode/Private/ScriptableToolsEditorModeModule.cpp
+++ b/Engine/Plugins/Editor/ScriptableToolsEditorMode/Source/ScriptableToolsEditorMode/Private/ScriptableToolsEditorModeModule.cpp
@@ -5,14 +5,25 @@
 #include "ScriptableToolsEditorModeManagerCommands.h"
 #include "ScriptableToolsEditorModeStyle.h"
 #include "Misc/CoreDelegates.h"
+#include "UI/ScriptableToolGroupSetCustomization.h"
+#include "PropertyEditorModule.h"
 
 #define LOCTEXT_NAMESPACE "FScriptableToolsEditorModeModule"
 
+
+static const FName PropertyEditorModuleName("PropertyEditor");
+static const FName ScriptableToolGroupSetName("ScriptableToolGroupSet");
+
+
 void FScriptableToolsEditorModeModule::StartupModule()
 {
 	// This code will execute after your module is loaded into memory; the exact timing is specified in the .uplugin file per-module
 
 	FCoreDelegates::OnPostEngineInit.AddRaw(this, &FScriptableToolsEditorModeModule::OnPostEngineInit);
+
+
+	FPropertyEditorModule& PropertyModule = FModuleManager::LoadModuleChecked<FPropertyEditorModule>(PropertyEditorModuleName);
+	PropertyModule.RegisterCustomPropertyTypeLayout(ScriptableToolGroupSetName, FOnGetPropertyTypeCustomizationInstance::CreateStatic(&FScriptableToolGroupSetCustomization::MakeInstance));
 }
 
 void FScriptableToolsEditorModeModule::ShutdownModule()
diff --git a/Engine/Plugins/Editor/ScriptableToolsEditorMode/Source/ScriptableToolsEditorMode/Private/ScriptableToolsEditorModeSettings.cpp b/Engine/Plugins/Editor/ScriptableToolsEditorMode/Source/ScriptableToolsEditorMode/Private/ScriptableToolsEditorModeSettings.cpp
new file mode 100644
--- /dev/null
+++ b/Engine/Plugins/Editor/ScriptableToolsEditorMode/Source/ScriptableToolsEditorMode/Private/ScriptableToolsEditorModeSettings.cpp
@@ -0,0 +1,22 @@
+// Copyright Epic Games, Inc. All Rights Reserved.
+
+#include "ScriptableToolsEditorModeSettings.h"
+
+#include UE_INLINE_GENERATED_CPP_BY_NAME(ScriptableToolsEditorModeSettings)
+
+
+#define LOCTEXT_NAMESPACE "ScriptableToolsEditorModeSettings"
+
+
+FText UScriptableToolsModeCustomizationSettings::GetSectionText() const
+{
+	return LOCTEXT("ScriptableModeSettingsName", "Scriptable Tools Mode");
+}
+
+FText UScriptableToolsModeCustomizationSettings::GetSectionDescription() const
+{
+	return LOCTEXT("ScriptableModeSettingsDescription", "Configure the Scriptable Tools Editor Mode plugin");
+}
+
+
+#undef LOCTEXT_NAMESPACE
\ No newline at end of file
diff --git a/Engine/Plugins/Editor/ScriptableToolsEditorMode/Source/ScriptableToolsEditorMode/Private/ScriptableToolsEditorModeStyle.cpp b/Engine/Plugins/Editor/ScriptableToolsEditorMode/Source/ScriptableToolsEditorMode/Private/ScriptableToolsEditorModeStyle.cpp
--- a/Engine/Plugins/Editor/ScriptableToolsEditorMode/Source/ScriptableToolsEditorMode/Private/ScriptableToolsEditorModeStyle.cpp
+++ b/Engine/Plugins/Editor/ScriptableToolsEditorMode/Source/ScriptableToolsEditorMode/Private/ScriptableToolsEditorModeStyle.cpp
@@ -9,9 +9,11 @@
 #include "Interfaces/IPluginManager.h"
 #include "SlateOptMacros.h"
 #include "Styling/SlateStyle.h"
+#include "Styling/StyleColors.h"
 
 
 #define IMAGE_PLUGIN_BRUSH( RelativePath, ... ) FSlateImageBrush( FScriptableToolsEditorModeStyle::InContent( RelativePath, ".png" ), __VA_ARGS__ )
+#define IMAGE_BRUSH_SVG( RelativePath, ... ) FSlateVectorImageBrush(FScriptableToolsEditorModeStyle::InContent( RelativePath, ".svg"), __VA_ARGS__)
 
 // This is to fix the issue that SlateStyleMacros like IMAGE_BRUSH look for RootToContentDir but StyleSet->RootToContentDir is how this style is set up
 #define RootToContentDir StyleSet->RootToContentDir
@@ -69,6 +71,8 @@ void FScriptableToolsEditorModeStyle::Initialize()
 
 	}
 
+	StyleSet->Set("ToolPalette.MenuIndicator", new IMAGE_BRUSH_SVG("Icons/chevron-right", Icon20x20, FStyleColors::Foreground));
+
 	FSlateStyleRegistry::RegisterSlateStyle(*StyleSet.Get());
 };
 
diff --git a/Engine/Plugins/Editor/ScriptableToolsEditorMode/Source/ScriptableToolsEditorMode/Private/ScriptableToolsEditorModeToolkit.cpp b/Engine/Plugins/Editor/ScriptableToolsEditorMode/Source/ScriptableToolsEditorMode/Private/ScriptableToolsEditorModeToolkit.cpp
--- a/Engine/Plugins/Editor/ScriptableToolsEditorMode/Source/ScriptableToolsEditorMode/Private/ScriptableToolsEditorModeToolkit.cpp
+++ b/Engine/Plugins/Editor/ScriptableToolsEditorMode/Source/ScriptableToolsEditorMode/Private/ScriptableToolsEditorModeToolkit.cpp
@@ -6,6 +6,8 @@
 #include "ScriptableToolsEditorModeManagerCommands.h"
 #include "ScriptableToolsEditorModeStyle.h"
 
+#include "SSingleObjectDetailsPanel.h"
+
 #include "Framework/MultiBox/MultiBoxBuilder.h"
 #include "IDetailsView.h"
 #include "Toolkits/AssetEditorModeUILayer.h"
@@ -15,14 +17,327 @@
 
 #include "Widgets/Text/STextBlock.h"
 #include "Widgets/Images/SImage.h"
+#include "Widgets/Layout/SWidgetSwitcher.h"
+#include "Widgets/Layout/SScrollBox.h"
+#include "Widgets/Layout/SExpandableArea.h"
+#include "Widgets/Docking/SDockTab.h"
+#include "Widgets/Notifications/SProgressBar.h"
 #include "SPrimaryButton.h"
+#include "SAssetDropTarget.h"
 
 #include "ScriptableInteractiveTool.h"
 #include "ScriptableToolSet.h"
+#include "ScriptableToolsEditorModeSettings.h"
+
+#include "UI/SScriptableToolGroupSetCombo.h"
+
+#include "Engine/Blueprint.h"
 
 #define LOCTEXT_NAMESPACE "FScriptableToolsEditorModeToolkit"
 
 
+class SToolPaletteLoadBar : public SCompoundWidget
+{
+public:
+	SLATE_BEGIN_ARGS(SToolPaletteLoadBar) {}
+	SLATE_END_ARGS()
+
+private:
+
+	FScriptableToolsEditorModeToolkit* Toolkit;
+
+	EVisibility IsVisible() const
+	{
+		if (Toolkit)
+		{
+			return Toolkit->AreToolsLoading() ? EVisibility::Visible : EVisibility::Collapsed;
+		}
+		return EVisibility::Collapsed;
+	}
+
+	TOptional<float> GetPercentLoaded() const
+	{
+		if (Toolkit)
+		{
+			return Toolkit->GetToolPercentLoaded();
+		}
+		return TOptional<float>();
+	}
+
+public:
+
+	void Construct(const FArguments& InArgs, FScriptableToolsEditorModeToolkit* ToolkitIn)
+	{
+		Toolkit = ToolkitIn;
+
+		ChildSlot
+		[
+			SNew(SBox)
+			.Visibility(this, &SToolPaletteLoadBar::IsVisible)
+			[
+				SNew(SVerticalBox)
+				+ SVerticalBox::Slot()
+				[
+					SNew(SHorizontalBox)
+					+ SHorizontalBox::Slot()
+					.Padding(2.f, 6.f)
+					//.AutoWidth()
+					[
+						SNew(STextBlock)
+						.Text(LOCTEXT("ScriptableToolsLoadingText", "Loading tools..."))
+					]
+				]
+				+SVerticalBox::Slot()
+				[
+					SNew(SHorizontalBox)
+					+ SHorizontalBox::Slot()
+					.Padding(2.f, 6.f)
+					//.AutoWidth()
+					[
+						SNew(SProgressBar)
+						.BorderPadding(FVector2D::ZeroVector)
+						.Percent(this, &SToolPaletteLoadBar::GetPercentLoaded)
+						.FillColorAndOpacity(FSlateColor(FLinearColor(0.0f, 1.0f, 1.0f)))
+					]
+				]
+			]
+
+		];
+
+	}
+
+};
+
+class SToolPaletteTagPanel : public SCompoundWidget
+{
+public:
+	SLATE_BEGIN_ARGS(SToolPaletteTagPanel) {}
+	SLATE_END_ARGS()
+
+private:
+	// Pointer back to owning sprite editor instance (the keeper of state)
+
+	UScriptableToolsModeCustomizationSettings* ModeSettings;
+	FScriptableToolsEditorModeToolkit* Toolkit;
+
+	TSharedPtr<SScriptableToolGroupSetCombo> TagCombo;
+	TSharedPtr<SComboButton> ToolMenuButton;
+
+	FDelegateHandle SettingsUpdateHandle;
+
+	void RefreshData(UObject*, FPropertyChangedEvent& )
+	{
+		if (TagCombo)
+		{
+			TagCombo->ForceUpdate();
+		}
+	}
+
+public:
+
+	virtual ~SToolPaletteTagPanel()
+	{
+		ModeSettings = GetMutableDefault<UScriptableToolsModeCustomizationSettings>();
+		ModeSettings->OnSettingChanged().Remove(SettingsUpdateHandle);
+	}
+
+	void Construct(const FArguments& InArgs, FScriptableToolsEditorModeToolkit* ToolkitIn)
+	{
+		ModeSettings = GetMutableDefault<UScriptableToolsModeCustomizationSettings>();
+		Toolkit = ToolkitIn;
+
+		SettingsUpdateHandle = ModeSettings->OnSettingChanged().AddSP(this, &SToolPaletteTagPanel::RefreshData);
+
+		ChildSlot
+		[
+			SAssignNew(ToolMenuButton, SComboButton)
+			.HasDownArrow(false)
+			.CollapseMenuOnParentFocus(false)
+			.MenuPlacement(EMenuPlacement::MenuPlacement_MenuRight)
+			.OnMenuOpenChanged_Lambda([this](bool bOpened)
+			{
+				FSlateThrottleManager::Get().DisableThrottle(bOpened);
+			})
+			.ButtonContent()
+			[
+				SNew(SAssetDropTarget)
+				.bSupportsMultiDrop(true)
+				.OnAreAssetsAcceptableForDropWithReason_Lambda([this](TArrayView<FAssetData> InAssets, FText& OutReason)
+				{
+					if (InAssets.Num() > 1)
+					{
+						OutReason = FText(LOCTEXT("ScriptableToolPaletteTagDropWarningPlural", "Assets must be Scriptable Tool Tags."));
+					}
+					else
+					{
+						OutReason = FText(LOCTEXT("ScriptableToolPaletteTagDropWarning", "Asset must be a Scriptable Tool Tag."));
+					}
+
+					for (FAssetData& Asset : InAssets)
+					{
+						UObject* AssetObject = Asset.GetAsset();
+
+						if (AssetObject == nullptr)
+						{
+							return false;
+						}
+
+						if (!AssetObject->IsA(UBlueprint::StaticClass()))
+						{
+							return false;
+						}
+
+						UBlueprint* BlueprintObject = Cast<UBlueprint>(AssetObject);
+
+						if (!BlueprintObject->GeneratedClass->IsChildOf(UScriptableToolGroupTag::StaticClass()))
+						{
+							return false;
+						}
+					}
+					return true;
+				})
+				.OnAssetsDropped_Lambda([this](const FDragDropEvent&, TArrayView<FAssetData> InAssets)
+					{
+						ModeSettings->PreEditChange(UScriptableToolsModeCustomizationSettings::StaticClass()->FindPropertyByName("bRegisterAllTools"));
+						ModeSettings->PreEditChange(UScriptableToolsModeCustomizationSettings::StaticClass()->FindPropertyByName("ToolRegistrationFilters"));
+						for (FAssetData& Asset : InAssets)
+						{
+							UObject* AssetObject = Asset.GetAsset();
+
+							UBlueprint* BlueprintObject = Cast<UBlueprint>(AssetObject);
+
+							UClass* BlueprintClass = BlueprintObject->GetBlueprintClass();							
+
+							if (BlueprintObject->GeneratedClass->IsChildOf(UScriptableToolGroupTag::StaticClass()))
+							{
+								TSubclassOf<UScriptableToolGroupTag> TagSubclass{ BlueprintObject->GeneratedClass };
+								ModeSettings->ToolRegistrationFilters.Groups.Add(TagSubclass);								
+							}
+						}
+						ModeSettings->PostEditChange();
+					})
+				[
+					SNew(SBorder)
+					.Visibility(EVisibility::SelfHitTestInvisible)
+					.Padding(FMargin(0.0f))
+					.BorderImage(FAppStyle::Get().GetBrush("ProjectBrowser.ProjectTile.DropShadow"))
+					[
+						SNew(SOverlay)
+						+ SOverlay::Slot()
+						[
+							SNew(SHorizontalBox)
+							+ SHorizontalBox::Slot()
+							.AutoWidth()
+							.HAlign(HAlign_Center)
+							.VAlign(VAlign_Center)
+							[
+								SNew(SBox)
+								.HeightOverride(30)
+								[
+									SNew(SBorder)
+									.Padding(0)
+									.BorderImage(FStyleDefaults::GetNoBrush())
+									.HAlign(HAlign_Center)
+									.VAlign(VAlign_Center)									
+									[
+										SNew(STextBlock)
+										.Text_Lambda([this]() {
+											if(ModeSettings->RegisterAllTools())
+											{
+												return LOCTEXT("ScriptableToolsAllToolsLabel", "Showing all tools");
+											}
+											else
+											{
+
+												TArray<FText> GroupNames;
+
+												for (UClass* GroupClass : ModeSettings->ToolRegistrationFilters.Groups)
+												{
+													if (GroupClass)
+													{
+														UScriptableToolGroupTag* GroupTag = Cast<UScriptableToolGroupTag>(GroupClass->GetDefaultObject());
+
+														if (GroupTag)
+														{
+															GroupNames.Add(FText::FromString(GroupTag->Name));
+														}
+													}
+												}
+
+												if (ModeSettings->ToolRegistrationFilters.Groups.Num() == 1)
+												{
+													ensure(GroupNames.IsValidIndex(0));
+													return FText::Format(LOCTEXT("ScriptableToolsOneGroupLabel", "Showing tools from {0}"), GroupNames[0]);
+												}
+												else if (ModeSettings->ToolRegistrationFilters.Groups.Num() == 2)
+												{
+													ensure(GroupNames.IsValidIndex(1));
+													return FText::Format(LOCTEXT("ScriptableToolsTwoGroupLabel", "Showing tools from {0} and {1}"), GroupNames[0], GroupNames[1]);
+												}
+												else
+												{
+													ensure(GroupNames.IsValidIndex(1));
+													return FText::Format(LOCTEXT("ScriptableToolsManyGroupLabel", "Showing tools from {0}, {1} and more..."), GroupNames[0], GroupNames[1]);
+												}
+											}
+										})
+										.ToolTipText(LOCTEXT("ScriptableToolsGroupButtonTooltip", "Select tool groups or drag tool group asset here to filter displayed tools."))
+									]
+								]
+							]
+							+ SHorizontalBox::Slot()
+							.FillWidth(1.0)
+							[
+								SNew(SBox)
+								.HeightOverride(30)
+							]
+							+ SHorizontalBox::Slot()
+							.AutoWidth()
+							.HAlign(EHorizontalAlignment::HAlign_Right)
+							[
+								SNew(SImage)
+								.Image(FScriptableToolsEditorModeStyle::Get()->GetBrush("ToolPalette.MenuIndicator"))
+							]
+						]
+					]
+				]
+			]			
+			.MenuContent()
+			[
+				SNew(SBox)				
+				.WidthOverride(300)
+				[
+					SNew(SBorder)
+					.Padding(15)
+					.BorderImage(FStyleDefaults::GetNoBrush())
+					[
+
+						SNew(SVerticalBox)
+						+ SVerticalBox::Slot()
+						.AutoHeight()
+						[
+							SNew(SHorizontalBox)
+							+ SHorizontalBox::Slot()
+							.FillWidth(1.0)
+							[
+								SAssignNew(TagCombo, SScriptableToolGroupSetCombo)
+								.StructPtr(&ModeSettings->ToolRegistrationFilters)
+								.OnChanged_Lambda([this]() {
+									ModeSettings->PostEditChange();
+								})
+							]
+						]
+					]
+				]
+			]			
+		];
+	}
+
+	
+
+	// End of SSingleObjectDetailsPanel interface
+};
+
 
 FScriptableToolsEditorModeToolkit::FScriptableToolsEditorModeToolkit()
 {
@@ -33,6 +348,9 @@ FScriptableToolsEditorModeToolkit::~FScriptableToolsEditorModeToolkit()
 {
 	GetScriptableEditorMode()->GetInteractiveToolsContext(EToolsContextScope::EdMode)->OnToolNotificationMessage.RemoveAll(this);
 	GetScriptableEditorMode()->GetInteractiveToolsContext(EToolsContextScope::EdMode)->OnToolWarningMessage.RemoveAll(this);
+
+	UScriptableToolsModeCustomizationSettings* ModeSettings = GetMutableDefault<UScriptableToolsModeCustomizationSettings>();
+	ModeSettings->OnSettingChanged().Remove(SettingsUpdateHandle);	
 }
 
 
@@ -466,7 +784,31 @@ void FScriptableToolsEditorModeToolkit::InvokeUI()
 void FScriptableToolsEditorModeToolkit::ForceToolPaletteRebuild()
 {
 	this->UpdateActiveToolCategories();
-	this->RebuildModeToolPalette();
+
+	if (ModeUILayer.IsValid() && HasIntegratedToolPalettes() == false)
+	{
+		if (TSharedPtr<FAssetEditorModeUILayer> ModeUILayerPtr = ModeUILayer.Pin())
+		{
+			TSharedPtr<FUICommandList> CommandList;
+			if (GetScriptableEditorMode().IsValid())
+			{
+				UEdMode* ScriptableMode = GetScriptableEditorMode().Get();
+				CommandList = GetToolkitCommands();
+				ActiveToolBarRows.Reset();
+
+				TArray<FName> PaletteNames;
+				GetToolPaletteNames(PaletteNames);
+				for (const FName& Palette : PaletteNames)
+				{
+					TSharedRef<SWidget> PaletteWidget = CreatePaletteWidget(CommandList, ScriptableMode->GetModeInfo().ToolbarCustomizationName, Palette);
+					ActiveToolBarRows.Emplace(ScriptableMode->GetID(), Palette, GetToolPaletteDisplayName(Palette), PaletteWidget);
+				}
+
+				RebuildModeToolPaletteWidgets();
+			}
+		}
+	}
+
 }
 
 
@@ -572,6 +914,140 @@ void FScriptableToolsEditorModeToolkit::UpdateShowWarnings()
 
 }
 
+void FScriptableToolsEditorModeToolkit::RebuildModeToolBar()
+{
+	TSharedPtr<SDockTab> ToolbarTabPtr = ModeToolbarTab.Pin();
+	if (ToolbarTabPtr && HasToolkitBuilder())
+	{
+		ToolbarTabPtr->SetParentDockTabStackTabWellHidden(true);
+	}
+
+	// If the tab or box is not valid the toolbar has not been opened or has been closed by the user
+	TSharedPtr<SVerticalBox> ModeToolbarBoxPinned = ModeToolbarBox.Pin();
+	if (ModeToolbarTab.IsValid() && ModeToolbarBoxPinned)
+	{
+		ModeToolbarBoxPinned->ClearChildren();
+		bool bExclusivePalettes = true;
+		ToolBoxVBox = SNew(SVerticalBox);
+
+		RebuildModeToolPaletteWidgets();
+
+		ModeToolbarBoxPinned->AddSlot()
+		.AutoHeight()
+		.Padding(1.f)
+		[
+			SNew(SBox)
+			[
+				SNew(SToolPaletteTagPanel, this)
+			]
+		];
+
+		ModeToolbarBoxPinned->AddSlot()
+		.AutoHeight()
+		.Padding(1.f)
+		[
+			SNew(SBox)
+			[
+				SNew(SToolPaletteLoadBar, this)
+			]
+		];
+
+		ModeToolbarBoxPinned->AddSlot()
+		[
+			SNew(SScrollBox)
+			.Visibility_Lambda([this]()
+			{
+				return AreToolsLoading() ? EVisibility::Collapsed : EVisibility::Visible;
+			})
+			+ SScrollBox::Slot()
+			[
+				ToolBoxVBox.ToSharedRef()
+			]
+		];
+		
+
+
+	}
+}
+
+void FScriptableToolsEditorModeToolkit::StartAsyncToolLoading()
+{
+	bAsyncLoadInProgress = true;
+	AsyncLoadProgress = 0.0;
+}
+
+void FScriptableToolsEditorModeToolkit::SetAsyncProgress(float PercentLoaded)
+{
+	ensure(bAsyncLoadInProgress == true);
+	AsyncLoadProgress = PercentLoaded;
+}
 
+void FScriptableToolsEditorModeToolkit::EndAsyncToolLoading()
+{
+	ensure(bAsyncLoadInProgress == true);
+	bAsyncLoadInProgress = false;
+	AsyncLoadProgress = 1.0;
+}
+
+bool FScriptableToolsEditorModeToolkit::AreToolsLoading() const
+{
+	return bAsyncLoadInProgress;
+}
+
+TOptional<float> FScriptableToolsEditorModeToolkit::GetToolPercentLoaded() const
+{
+	if (bAsyncLoadInProgress)
+	{
+		return TOptional<float>(AsyncLoadProgress);
+	}
+	return TOptional<float>();
+}
+
+bool FScriptableToolsEditorModeToolkit::ShouldShowModeToolbar() const
+{
+	// Always return true here, to ensure we render the toolbar even if no tools are loaded via group tags.
+	return true;
+}
+
+
+void FScriptableToolsEditorModeToolkit::RebuildModeToolPaletteWidgets()
+{
+	if (ToolBoxVBox)
+	{
+		ToolBoxVBox->ClearChildren();
+
+		int32 PaletteCount = ActiveToolBarRows.Num();
+		if (PaletteCount > 0)
+		{
+			for (int32 RowIdx = 0; RowIdx < PaletteCount; ++RowIdx)
+			{
+				const FEdModeToolbarRow& Row = ActiveToolBarRows[RowIdx];
+				if (ensure(Row.ToolbarWidget.IsValid()))
+				{
+					TSharedRef<SWidget> PaletteWidget = Row.ToolbarWidget.ToSharedRef();
+
+					ToolBoxVBox->AddSlot()
+						.AutoHeight()
+						.Padding(FMargin(2.0, 2.0))
+						[
+							SNew(SExpandableArea)
+							.AreaTitle(Row.DisplayName)
+						.AreaTitleFont(FAppStyle::Get().GetFontStyle("NormalFont"))
+						.BorderImage(FAppStyle::Get().GetBrush("PaletteToolbar.ExpandableAreaHeader"))
+						.BodyBorderImage(FAppStyle::Get().GetBrush("PaletteToolbar.ExpandableAreaBody"))
+						.HeaderPadding(FMargin(4.f))
+						.Padding(FMargin(4.0, 0.0))
+						.BodyContent()
+						[
+							PaletteWidget
+						]
+						];
+
+				}
+			}
+		}
+	}
+
+}
 
 #undef LOCTEXT_NAMESPACE
diff --git a/Engine/Plugins/Editor/ScriptableToolsEditorMode/Source/ScriptableToolsEditorMode/Private/Tags/EditableScriptableToolGroupSet.cpp b/Engine/Plugins/Editor/ScriptableToolsEditorMode/Source/ScriptableToolsEditorMode/Private/Tags/EditableScriptableToolGroupSet.cpp
new file mode 100644
--- /dev/null
+++ b/Engine/Plugins/Editor/ScriptableToolsEditorMode/Source/ScriptableToolsEditorMode/Private/Tags/EditableScriptableToolGroupSet.cpp
@@ -0,0 +1,34 @@
+// Copyright Epic Games, Inc. All Rights Reserved.
+
+#include "Tags/EditableScriptableToolGroupSet.h"
+
+UEditableScriptableToolGroupSet::UEditableScriptableToolGroupSet()
+	: GroupsProperty(nullptr)
+{
+	if (!IsTemplate())
+	{
+		GroupsProperty = FindFProperty<FProperty>(GetClass(), TEXT("GroupSet"));
+	}
+}
+
+void UEditableScriptableToolGroupSet::SetGroups(const FScriptableToolGroupSet::FGroupSet& InGroups)
+{
+	GroupSet.Groups = InGroups;
+}
+
+FScriptableToolGroupSet::FGroupSet& UEditableScriptableToolGroupSet::GetGroups()
+{
+	return GroupSet.Groups;
+}
+
+FString UEditableScriptableToolGroupSet::GetGroupSetExportText()
+{
+	GroupsPropertyAsString.Reset();
+
+	if (GroupsProperty)
+	{
+		GroupsProperty->ExportTextItem_Direct(GroupsPropertyAsString, &GroupSet, &GroupSet, this, 0);
+	}
+
+	return GroupsPropertyAsString;
+}
\ No newline at end of file
diff --git a/Engine/Plugins/Editor/ScriptableToolsEditorMode/Source/ScriptableToolsEditorMode/Private/UI/SScriptableToolGroupSetCombo.cpp b/Engine/Plugins/Editor/ScriptableToolsEditorMode/Source/ScriptableToolsEditorMode/Private/UI/SScriptableToolGroupSetCombo.cpp
new file mode 100644
--- /dev/null
+++ b/Engine/Plugins/Editor/ScriptableToolsEditorMode/Source/ScriptableToolsEditorMode/Private/UI/SScriptableToolGroupSetCombo.cpp
@@ -0,0 +1,252 @@
+// Copyright Epic Games, Inc. All Rights Reserved.
+
+#include "UI/SScriptableToolGroupSetCombo.h"
+
+#include "Tags/ScriptableToolGroupSet.h"
+#include "Tags/ScriptableToolGroupTag.h"
+#include "Tags/EditableScriptableToolGroupSet.h"
+#include "PropertyHandle.h"
+#include "Framework/Views/TableViewMetadata.h"
+#include "UI/SScriptableToolGroupSetPicker.h"
+#include "UI/SScriptableToolGroupTagChip.h"
+#include "UObject/Package.h"
+#include "Widgets/Text/STextBlock.h"
+#include "Widgets/Layout/SScrollBar.h"
+#include "Widgets/Views/SListView.h"
+#include "Widgets/Input/SComboButton.h"
+
+#define LOCTEXT_NAMESPACE "ScriptableToolGroupSetCombo"
+
+//------------------------------------------------------------------------------
+// SScriptableToolGroupSetCombo
+//------------------------------------------------------------------------------
+
+SLATE_IMPLEMENT_WIDGET(SScriptableToolGroupSetCombo)
+void SScriptableToolGroupSetCombo::PrivateRegisterAttributes(FSlateAttributeInitializer& AttributeInitializer)
+{
+
+}
+
+void SScriptableToolGroupSetCombo::Construct(const FArguments& InArgs)
+{
+	StructPropertyHandle = InArgs._StructPropertyHandle;
+	StructPtr = InArgs._StructPtr;
+	OnChanged = InArgs._OnChanged;
+
+	HelperGroupSet.Reset(NewObject<UEditableScriptableToolGroupSet>(GetTransientPackage(), NAME_None, RF_Transient));
+
+	if (StructPropertyHandle.IsValid() && StructPropertyHandle->IsValidHandle() && !StructPtr)
+	{
+		StructPropertyHandle->SetOnPropertyValueChanged(FSimpleDelegate::CreateSP(this, &SScriptableToolGroupSetCombo::RefreshListView));
+	}
+
+	RefreshListView();
+
+	TWeakPtr<SScriptableToolGroupSetCombo> WeakSelf = StaticCastWeakPtr<SScriptableToolGroupSetCombo>(AsWeak());
+
+	ActiveGroupTagsListView = SNew(SListView<UClass*>)
+		.ListItemsSource(&ActiveGroupTags)
+		.SelectionMode(ESelectionMode::None)
+		.ListViewStyle(&FAppStyle::Get().GetWidgetStyle<FTableViewStyle>("SimpleListView"))
+		.OnGenerateRow(this, &SScriptableToolGroupSetCombo::OnGenerateRow)
+		.Visibility_Lambda([WeakSelf]()
+			{
+				if (const TSharedPtr<SScriptableToolGroupSetCombo> Self = WeakSelf.Pin())
+				{
+					return Self->ActiveGroupTags.Num() > 0 ? EVisibility::Visible : EVisibility::Collapsed;
+				}
+				return EVisibility::Collapsed;
+			});
+
+	ChildSlot
+		[
+			SNew(SHorizontalBox)
+
+			+ SHorizontalBox::Slot()
+		.FillWidth(1)
+		.VAlign(VAlign_Top)
+		[
+			SAssignNew(ComboButton, SComboButton)
+			.HasDownArrow(true)
+		.VAlign(VAlign_Top)
+		.ContentPadding(0)
+		.OnGetMenuContent(this, &SScriptableToolGroupSetCombo::OnGetMenuContent)
+		.CollapseMenuOnParentFocus(true)
+		.ButtonContent()
+		[
+			SNew(SHorizontalBox)
+
+			+ SHorizontalBox::Slot()
+			.FillWidth(1)
+			.VAlign(VAlign_Top)
+			[
+				SNew(SHorizontalBox)
+
+				// Group Tag List
+				+ SHorizontalBox::Slot()
+				.VAlign(VAlign_Top)
+				.FillWidth(1)
+				[
+					ActiveGroupTagsListView.ToSharedRef()
+				]
+
+				// Empty indicator
+				+ SHorizontalBox::Slot()
+				.VAlign(VAlign_Center)
+				.FillWidth(1)
+				.Padding(FMargin(4, 2))
+				[
+					SNew(SBox)
+					.HeightOverride(SScriptableToolGroupTagChip::ChipHeight)
+					.VAlign(VAlign_Center)
+					.Padding(0, 0, 8, 0)
+					.Visibility_Lambda([WeakSelf]()
+					{
+						if (const TSharedPtr<SScriptableToolGroupSetCombo> Self = WeakSelf.Pin())
+						{	
+							return Self->ActiveGroupTags.Num() > 0 ? EVisibility::Collapsed : EVisibility::Visible;
+						}
+						return EVisibility::Collapsed;
+					})
+					[
+						SNew(STextBlock)
+						.ColorAndOpacity(FSlateColor::UseSubduedForeground())
+						.Font(FAppStyle::GetFontStyle(TEXT("PropertyWindow.NormalFont")))
+						.Text(LOCTEXT("ScriptableToolGroupSetCombo_Empty", "Select Tool Groups..."))
+						.ToolTipText(LOCTEXT("ScriptableToolGroupSetCombo_EmptyTooltip", "No Tool Groups selected. Use dropdown to load tools into mode palette."))
+					]
+				]
+			]
+		]
+		]
+		];
+	
+}
+
+TSharedRef<ITableRow> SScriptableToolGroupSetCombo::OnGenerateRow(UClass* InGroup, const TSharedRef<STableViewBase>& OwnerTable)
+{
+	return SNew(STableRow<UClass*>, OwnerTable)
+		.Style(&FAppStyle::Get().GetWidgetStyle<FTableRowStyle>("SimpleTableView.Row"))
+		.Padding(FMargin(0, 2))		
+		[
+			SNew(SHorizontalBox)
+			+SHorizontalBox::Slot()
+			.FillWidth(1.0)
+			[
+				SNew(SScriptableToolGroupTagChip)
+				.TagClass(InGroup)
+				.Text(FText::FromString(Cast<UScriptableToolGroupTag>(InGroup->GetDefaultObject())->Name))
+				.ToolTipText(FText::FromString(InGroup->GetClassPathName().ToString()))
+				.OnClearPressed(this, &SScriptableToolGroupSetCombo::OnClearTagClicked, InGroup)
+			]
+		];
+}
+
+TSharedRef<SWidget> SScriptableToolGroupSetCombo::OnGetMenuContent()
+{
+	TagPicker = SNew(SScriptableToolGroupSetPicker)
+		.StructPropertyHandle(StructPropertyHandle)
+		.StructPtr(StructPtr)
+		.OnChanged_Lambda([this]() {
+			OnChanged.ExecuteIfBound();
+		});
+
+	ComboButton->SetMenuContentWidgetToFocus(TagPicker);
+
+	return TagPicker.ToSharedRef();
+}
+
+FReply SScriptableToolGroupSetCombo::OnClearTagClicked(UClass* InGroup)
+{
+	HelperGroupSet->GetGroups().Remove(TSubclassOf<UScriptableToolGroupTag>(InGroup));
+
+	if (StructPropertyHandle.IsValid() && StructPropertyHandle->IsValidHandle() && !StructPtr)
+	{
+		// Set the property with a formatted string in order to propagate CDO changes to instances if necessary
+		const FString OutString = HelperGroupSet->GetGroupSetExportText();
+		StructPropertyHandle->SetValueFromFormattedString(OutString);
+	}
+	else if (StructPtr)
+	{
+		StructPtr->Groups = HelperGroupSet->GetGroups();
+	}
+
+	OnChanged.ExecuteIfBound();
+	RefreshListView();
+
+	return FReply::Handled();
+}
+
+void SScriptableToolGroupSetCombo::Tick(const FGeometry& AllottedGeometry, const double InCurrentTime, const float InDeltaTime)
+{
+	if (!IsEnabled() && ComboButton->IsOpen())
+	{
+		ComboButton->SetIsOpen(false);
+	}
+}
+
+void SScriptableToolGroupSetCombo::ForceUpdate()
+{
+	RefreshListView();
+	if (TagPicker.IsValid())
+	{
+		TagPicker->ForceUpdate();
+	}
+}
+
+void SScriptableToolGroupSetCombo::RefreshListView()
+{
+	ActiveGroupTags.Reset();
+
+	// Add UClass* to our TagsToEdit list from the property handle
+	void* StructPointer = nullptr;
+	bool bValidData = false;
+	if (StructPropertyHandle.IsValid() && StructPropertyHandle->IsValidHandle() && !StructPtr)
+	{
+		if (StructPropertyHandle->GetValueData(StructPointer) == FPropertyAccess::Success && StructPointer)
+		{
+			bValidData = true;
+		}
+	}
+	else if (StructPtr)
+	{
+		StructPointer = StructPtr;
+		bValidData = true;
+	}
+
+	if(bValidData)
+	{
+		FScriptableToolGroupSet& GroupSet = *static_cast<FScriptableToolGroupSet*>(StructPointer);
+		FScriptableToolGroupSet::FGroupSet& Groups = GroupSet.Groups;
+
+		HelperGroupSet->SetGroups(Groups);
+
+		for (TSubclassOf<UScriptableToolGroupTag>& Element : Groups)
+		{
+			if (*Element)
+			{
+				ActiveGroupTags.AddUnique(*Element);
+			}
+		}
+	}
+
+	// Lexicographically sort Group tags.
+	Algo::Sort(ActiveGroupTags, [](const UClass* LHS, const UClass* RHS)
+		{
+			const TObjectPtr<UScriptableToolGroupTag> LHSCDO = Cast<UScriptableToolGroupTag>(LHS->GetDefaultObject());
+			const TObjectPtr<UScriptableToolGroupTag> RHSCDO = Cast<UScriptableToolGroupTag>(RHS->GetDefaultObject());
+
+			check(IsValid(LHSCDO) && IsValid(RHSCDO));
+
+			return LHSCDO->Name < RHSCDO->Name;
+		});
+
+	// Refresh the slate list
+	if (ActiveGroupTagsListView.IsValid())
+	{
+		ActiveGroupTagsListView->SetItemsSource(&ActiveGroupTags);
+		ActiveGroupTagsListView->RequestListRefresh();
+	}
+}
+
+#undef LOCTEXT_NAMESPACE
diff --git a/Engine/Plugins/Editor/ScriptableToolsEditorMode/Source/ScriptableToolsEditorMode/Private/UI/SScriptableToolGroupSetPicker.cpp b/Engine/Plugins/Editor/ScriptableToolsEditorMode/Source/ScriptableToolsEditorMode/Private/UI/SScriptableToolGroupSetPicker.cpp
new file mode 100644
--- /dev/null
+++ b/Engine/Plugins/Editor/ScriptableToolsEditorMode/Source/ScriptableToolsEditorMode/Private/UI/SScriptableToolGroupSetPicker.cpp
@@ -0,0 +1,333 @@
+// Copyright Epic Games, Inc. All Rights Reserved.
+
+#include "UI/SScriptableToolGroupSetPicker.h"
+
+#include "Algo/ForEach.h"
+#include "AssetRegistry/ARFilter.h"
+#include "Tags/ScriptableToolGroupTag.h"
+#include "Tags/EditableScriptableToolGroupSet.h"
+#include "Framework/Commands/UIAction.h"
+#include "Framework/MultiBox/MultiBoxBuilder.h"
+#include "Framework/Views/TableViewMetadata.h"
+#include "PropertyHandle.h"
+#include "UI/SScriptableToolGroupTagChip.h"
+#include "Textures/SlateIcon.h"
+#include "Widgets/DeclarativeSyntaxSupport.h"
+#include "Widgets/Input/SSearchBox.h"
+#include "Widgets/Input/SCheckBox.h"
+
+// Asset registry querying
+#include "AssetRegistry/IAssetRegistry.h"
+#include "Blueprint/BlueprintSupport.h"	// For FBlueprintTags::NativeParentClassPath
+#include "AssetRegistry/AssetData.h"
+#include "Engine/Blueprint.h"
+#include "Widgets/Layout/SScrollBar.h"
+#include "Widgets/Views/SListView.h"
+
+#include "Kismet2/KismetEditorUtilities.h"
+#include "AssetRegistry/AssetRegistryModule.h"
+
+#define LOCTEXT_NAMESPACE "ScriptableToolGroupSetPicker"
+
+//------------------------------------------------------------------------------
+// SScriptableToolGroupSetPicker
+//------------------------------------------------------------------------------
+
+void SScriptableToolGroupSetPicker::Construct(const FArguments& InArgs)
+{
+	StructPropertyHandle = InArgs._StructPropertyHandle;
+	StructPtr = InArgs._StructPtr;
+	OnChanged = InArgs._OnChanged;
+
+	if (StructPropertyHandle.IsValid() && StructPropertyHandle->IsValidHandle() && !StructPtr)
+	{
+		StructPropertyHandle->SetOnPropertyValueChanged(FSimpleDelegate::CreateSP(this, &SScriptableToolGroupSetPicker::PopulateCheckedTags));
+	}
+
+	PopulateVisibleClasses();
+
+	HelperGroupSet.Reset(NewObject<UEditableScriptableToolGroupSet>(GetTransientPackage(), NAME_None, RF_Transient));
+
+	PopulateCheckedTags();
+
+	OnSearchStringChanged(FText());
+
+	ChildSlot
+		[
+			GetChildWidget()
+		];	
+}
+
+TSharedRef<SWidget> SScriptableToolGroupSetPicker::GetChildWidget()
+{
+	const TSharedRef<SWidget> MenuContent = SNew(SBox)
+		[
+			SNew(SVerticalBox)
+
+			// Search box
+		+ SVerticalBox::Slot()
+		.AutoHeight()
+		.VAlign(VAlign_Top)
+		[
+			SAssignNew(GroupSearchBox, SSearchBox)
+			.HintText(LOCTEXT("ScriptableToolGroupSetPicker_SearchBoxHint", "Search Tool Group Tags"))
+		.OnTextChanged(this, &SScriptableToolGroupSetPicker::OnSearchStringChanged)
+		]
+
+	// List of tags
+	+ SVerticalBox::Slot()
+		.AutoHeight()
+		[
+			SAssignNew(VisibleGroupTagsListView, SListView<UClass*>)
+			.ListItemsSource(&VisibleGroupTags)
+		.OnGenerateRow(this, &SScriptableToolGroupSetPicker::OnGenerateRow)
+		.SelectionMode(ESelectionMode::None)
+		.ListViewStyle(&FAppStyle::Get().GetWidgetStyle<FTableViewStyle>("SimpleListView"))
+		]
+		];
+
+
+
+	FMenuBuilder MenuBuilder(/*bInShouldCloseWindowAfterMenuSelection*/false, nullptr);
+
+	MenuBuilder.BeginSection(FName(), LOCTEXT("SectionGroupSet", "Tool Group Tags"));
+
+	MenuBuilder.AddMenuEntry(
+		LOCTEXT("ScriptableToolGroupSetPicker_CreateNewTag", "Create New Tag"),
+		FText::GetEmpty(), FSlateIcon(FAppStyle::GetAppStyleSetName(), "Icons.X"),
+		FUIAction(FExecuteAction::CreateRaw(this, &SScriptableToolGroupSetPicker::OnCreateNewTag))
+	);
+
+	MenuBuilder.AddMenuEntry(
+		LOCTEXT("ScriptableToolGroupSetPicker_ClearAllTags", "Clear All Tags"), FText::GetEmpty(), FSlateIcon(FAppStyle::GetAppStyleSetName(), "Icons.X"),
+		FUIAction(FExecuteAction::CreateRaw(this, &SScriptableToolGroupSetPicker::OnUncheckAllTags))
+	);
+
+	MenuBuilder.AddSeparator();
+
+	MenuBuilder.AddWidget(MenuContent, FText::GetEmpty(), true);
+
+	MenuBuilder.EndSection();
+
+
+
+	return MenuBuilder.MakeWidget();
+}
+
+TSharedRef<ITableRow> SScriptableToolGroupSetPicker::OnGenerateRow(UClass* InTag, const TSharedRef<STableViewBase>& OwnerTable)
+{
+	if (const TObjectPtr<UScriptableToolGroupTag> SubclassCDO = Cast<UScriptableToolGroupTag>(InTag->GetDefaultObject()))
+	{
+		return SNew(STableRow<UClass*>, OwnerTable)
+			[
+				SNew(SCheckBox)
+				.OnCheckStateChanged_Lambda([this, InTag](ECheckBoxState NewState)
+					{
+						if (NewState == ECheckBoxState::Checked)
+						{
+							OnTagChecked(InTag);
+						}
+						else if (NewState == ECheckBoxState::Unchecked)
+						{
+							OnTagUnchecked(InTag);
+						}
+					})
+			.IsChecked_Lambda([this, InTag]()
+				{
+					return IsTagChecked(InTag);
+				})
+						.ToolTipText(FText::FromString(InTag->GetClassPathName().ToString()))
+					.Content()
+					[
+						SNew(STextBlock)
+						.Text(FText::FromString(SubclassCDO->Name))
+					]
+			];
+	}
+
+	return SNew(STableRow<UClass*>, OwnerTable);
+}
+
+void SScriptableToolGroupSetPicker::OnSearchStringChanged(const FText& NewString)
+{
+	SearchString = NewString.ToString();
+	RefreshListView();
+}
+
+void SScriptableToolGroupSetPicker::RefreshListView()
+{
+	VisibleGroupTags.Empty();
+
+	// Add all tags matching the search string.
+	for (UClass* Subclass : AllGroupTags)
+	{
+		if (const TObjectPtr<UScriptableToolGroupTag> SubclassCDO = Cast<UScriptableToolGroupTag>(Subclass->GetDefaultObject()))
+		{
+			if (SearchString.IsEmpty() || SubclassCDO->Name.Contains(SearchString))
+			{
+				VisibleGroupTags.AddUnique(Subclass);
+			}
+		}
+	}
+
+	// Lexicographically sort group tags.
+	Algo::Sort(VisibleGroupTags, [](const UClass* LHS, const UClass* RHS)
+		{
+			const TObjectPtr<UScriptableToolGroupTag> LHSCDO = Cast<UScriptableToolGroupTag>(LHS->GetDefaultObject());
+			const TObjectPtr<UScriptableToolGroupTag> RHSCDO = Cast<UScriptableToolGroupTag>(RHS->GetDefaultObject());
+
+			check(IsValid(LHSCDO) && IsValid(RHSCDO));
+
+			return LHSCDO->Name < RHSCDO->Name;
+		});
+
+	// Refresh the slate list.
+	if (VisibleGroupTagsListView.IsValid())
+	{
+		VisibleGroupTagsListView->SetItemsSource(&VisibleGroupTags);
+		VisibleGroupTagsListView->RequestListRefresh();
+	}
+}
+
+void SScriptableToolGroupSetPicker::PopulateVisibleClasses()
+{
+	const IAssetRegistry& AssetRegistry = IAssetRegistry::GetChecked();
+	const FTopLevelAssetPath BaseAssetPath = UScriptableToolGroupTag::StaticClass()->GetClassPathName();
+
+	AllGroupTags.Reset();
+
+	FARFilter Filter;
+	Filter.ClassPaths.Add(UBlueprint::StaticClass()->GetClassPathName());
+	Filter.bRecursiveClasses = true;
+	Filter.TagsAndValues.Add(FBlueprintTags::NativeParentClassPath);
+
+	AssetRegistry.EnumerateAssets(Filter, [this, BaseAssetPath](const FAssetData& Data)
+		{
+			if (Data.TagsAndValues.FindTag(FBlueprintTags::NativeParentClassPath).AsExportPath().ToTopLevelAssetPath() == BaseAssetPath)
+			{
+				if (const UBlueprint* AssetAsBlueprint = Cast<UBlueprint>(Data.GetAsset()))
+				{
+					if (const TSubclassOf<UScriptableToolGroupTag> GeneratedClassAsGroupClass = *AssetAsBlueprint->GeneratedClass)
+					{
+						AllGroupTags.Add(GeneratedClassAsGroupClass);
+					}
+				}
+			}
+			return true;	// Returning false will halt the enumeration
+		});
+}
+
+void SScriptableToolGroupSetPicker::ForceUpdate()
+{
+	PopulateCheckedTags();
+}
+
+void SScriptableToolGroupSetPicker::PopulateCheckedTags()
+{
+	// Reset checked tag map
+	for (UClass* Group : AllGroupTags)
+	{
+		CheckedTags.FindOrAdd(Group) = false;
+	}
+
+	// Access group set array
+	void* StructPointer = nullptr;
+	bool bValidData = false;
+	if (StructPropertyHandle.IsValid() && StructPropertyHandle->IsValidHandle() && !StructPtr)
+	{
+		if (StructPropertyHandle->GetValueData(StructPointer) == FPropertyAccess::Success && StructPointer)
+		{
+			bValidData = true;
+		}
+	}
+	else if (StructPtr)
+	{
+		StructPointer = StructPtr;
+		bValidData = true;
+	}
+
+	if (bValidData)
+	{
+		FScriptableToolGroupSet& GroupSet = *static_cast<FScriptableToolGroupSet*>(StructPointer);
+		FScriptableToolGroupSet::FGroupSet& Groups = GroupSet.Groups;
+
+		HelperGroupSet->SetGroups(Groups);
+
+		// Bring checked tag map to parity with current group set
+		for (TSubclassOf<UScriptableToolGroupTag>& Element : Groups)
+		{
+			if (CheckedTags.Contains(Element))
+			{
+				CheckedTags[Element] = true;
+			}
+		}
+	}
+}
+
+ECheckBoxState SScriptableToolGroupSetPicker::IsTagChecked(UClass* InTag)
+{
+	// This will add InTag to the map if not in it with a default value of false.
+	return CheckedTags.FindOrAdd(InTag) ? ECheckBoxState::Checked : ECheckBoxState::Unchecked;
+}
+
+void SScriptableToolGroupSetPicker::OnTagChecked(UClass* InTag)
+{
+	CheckedTags.FindOrAdd(InTag) = true;
+	HelperGroupSet->GetGroups().FindOrAdd(InTag);	// Expected value defaults to true when adding a group.
+	FlushHelperGroupSet();
+}
+
+void SScriptableToolGroupSetPicker::OnTagUnchecked(UClass* InTag)
+{
+	CheckedTags.FindOrAdd(InTag) = false;
+	HelperGroupSet->GetGroups().Remove(InTag);
+	FlushHelperGroupSet();
+}
+
+void SScriptableToolGroupSetPicker::OnUncheckAllTags()
+{
+	Algo::ForEach(AllGroupTags, [this](UClass* Group) { CheckedTags.FindOrAdd(Group) = false; });
+	HelperGroupSet->GetGroups().Empty();
+	FlushHelperGroupSet();
+}
+
+void SScriptableToolGroupSetPicker::OnCreateNewTag()
+{
+
+	UBlueprint* NewBlueprint = FKismetEditorUtilities::CreateBlueprintFromClass(
+		LOCTEXT("ScriptableToolGroupSetPicker_CreateTagAssetDialog","Create a new ScriptableTool group tag"),
+		UScriptableToolGroupTag::StaticClass(),
+		TEXT("NewTag"));
+
+	if (NewBlueprint)
+	{
+		// Editing the BluePrint
+		UScriptableToolGroupTag* NewBlueprintDefaultValues = NewBlueprint->GeneratedClass->GetDefaultObject<UScriptableToolGroupTag>();
+		NewBlueprintDefaultValues->Name = FString(NewBlueprintDefaultValues->GetName()).LeftChop(2).RightChop(9); // Taking off the "Default__<Name>_C" bits
+
+		// Compile the changes
+		FCompilerResultsLog LogResults;
+		FKismetEditorUtilities::CompileBlueprint(NewBlueprint, EBlueprintCompileOptions::None, &LogResults);
+
+		// Notify the asset registry
+		FAssetRegistryModule::AssetCreated(NewBlueprint);
+	}
+}
+
+void SScriptableToolGroupSetPicker::FlushHelperGroupSet() const
+{
+	if (StructPropertyHandle.IsValid() && StructPropertyHandle->IsValidHandle() && !StructPtr)
+	{
+		// Set the property with a formatted string in order to propagate CDO changes to instances if necessary
+		const FString OutString = HelperGroupSet->GetGroupSetExportText();
+		StructPropertyHandle->SetValueFromFormattedString(OutString);
+	}
+	else if (StructPtr)
+	{
+		StructPtr->Groups = HelperGroupSet->GetGroups();
+	}
+
+	OnChanged.ExecuteIfBound();
+}
+
+#undef LOCTEXT_NAMESPACE
diff --git a/Engine/Plugins/Editor/ScriptableToolsEditorMode/Source/ScriptableToolsEditorMode/Private/UI/SScriptableToolGroupTagChip.cpp b/Engine/Plugins/Editor/ScriptableToolsEditorMode/Source/ScriptableToolsEditorMode/Private/UI/SScriptableToolGroupTagChip.cpp
new file mode 100644
--- /dev/null
+++ b/Engine/Plugins/Editor/ScriptableToolsEditorMode/Source/ScriptableToolsEditorMode/Private/UI/SScriptableToolGroupTagChip.cpp
@@ -0,0 +1,123 @@
+// Copyright Epic Games, Inc. All Rights Reserved.
+
+#include "UI/SScriptableToolGroupTagChip.h"
+
+#include "Styling/StyleColors.h"
+#include "Widgets/DeclarativeSyntaxSupport.h"
+#include "Widgets/Layout/SBox.h"
+#include "Widgets/Input/SButton.h"
+#include "Widgets/Input/SCheckBox.h"
+#include "Widgets/SBoxPanel.h"
+#include "Widgets/Text/STextBlock.h"
+#include "Widgets/Images/SImage.h"
+
+#define LOCTEXT_NAMESPACE "ScriptableToolGroupTagChip"
+
+//------------------------------------------------------------------------------
+// SScriptableToolGroupTagChip
+//------------------------------------------------------------------------------
+
+SLATE_IMPLEMENT_WIDGET(SScriptableToolGroupTagChip)
+void SScriptableToolGroupTagChip::PrivateRegisterAttributes(FSlateAttributeInitializer& AttributeInitializer)
+{
+	SLATE_ADD_MEMBER_ATTRIBUTE_DEFINITION_WITH_NAME(AttributeInitializer, "Text", TextAttribute, EInvalidateWidgetReason::Layout);
+	SLATE_ADD_MEMBER_ATTRIBUTE_DEFINITION_WITH_NAME(AttributeInitializer, "ToolTipText", ToolTipTextAttribute, EInvalidateWidgetReason::Layout);
+}
+
+SScriptableToolGroupTagChip::SScriptableToolGroupTagChip()
+	: ToolTipTextAttribute(*this)
+	, TextAttribute(*this)
+{}
+
+void SScriptableToolGroupTagChip::Construct(const FArguments& InArgs)
+{
+	ToolTipTextAttribute.Assign(*this, InArgs._ToolTipText);
+	TextAttribute.Assign(*this, InArgs._Text);
+	OnClearPressed = InArgs._OnClearPressed;
+	TagClass = InArgs._TagClass;
+
+	TWeakPtr<SScriptableToolGroupTagChip> WeakSelf = StaticCastWeakPtr<SScriptableToolGroupTagChip>(AsWeak());
+
+	ChildSlot
+		[
+			SNew(SBox)
+			.HeightOverride(ChipHeight)		
+		[
+			SNew(SButton)			
+			.VAlign(VAlign_Center)
+		.HAlign(HAlign_Fill)
+		.OnClicked_Lambda([WeakSelf]()
+			{
+				return FReply::Unhandled();
+			})
+		[
+			SNew(SHorizontalBox)
+
+			+ SHorizontalBox::Slot()
+				.FillWidth(1.0)
+				.VAlign(VAlign_Center)
+				.HAlign(HAlign_Center)
+				.Padding(10.f, 0.f)
+				[
+					SNew(STextBlock)
+					.Font(FAppStyle::GetFontStyle(TEXT("PropertyWindow.NormalFont")))
+				.Text_Lambda([WeakSelf]()
+					{
+						const TSharedPtr<SScriptableToolGroupTagChip> Self = WeakSelf.Pin();
+						return Self.IsValid() ? Self->TextAttribute.Get() : FText::GetEmpty();
+					})
+				.ToolTipText_Lambda([WeakSelf]()
+					{
+						const TSharedPtr<SScriptableToolGroupTagChip> Self = WeakSelf.Pin();
+						return Self.IsValid() ? Self->ToolTipTextAttribute.Get() : FText::GetEmpty();
+					})
+				]
+
+			// Clear Button
+			+ SHorizontalBox::Slot()
+				.AutoWidth()
+				.VAlign(VAlign_Center)
+				.HAlign(HAlign_Right)
+				.Padding(0.f)
+				[
+					SAssignNew(ClearButton, SButton)
+					.Visibility_Lambda([WeakSelf]()
+						{
+							const TSharedPtr<SScriptableToolGroupTagChip> Self = WeakSelf.Pin();
+							return Self.IsValid() ? EVisibility::Visible : EVisibility::Collapsed;
+						})
+				.VAlign(VAlign_Center)
+							.HAlign(HAlign_Center)
+							.ToolTipText(LOCTEXT("ClearTag", "Clear Tag"))
+							.ContentPadding(0)
+							.OnClicked_Lambda([WeakSelf]()
+								{
+									//UE_LOG(LogScriptableToolEditor, Log, TEXT("Clear button clicked!"))
+										const TSharedPtr<SScriptableToolGroupTagChip> Self = WeakSelf.Pin();
+									if (Self.IsValid() && Self->OnClearPressed.IsBound())
+									{
+										return Self->OnClearPressed.Execute();
+									}
+									return FReply::Unhandled();
+								})
+							[
+								SNew(SImage)
+								.ColorAndOpacity_Lambda([WeakSelf]()
+									{
+										const TSharedPtr<SScriptableToolGroupTagChip> Self = WeakSelf.Pin();
+										if (Self.IsValid() && Self->ClearButton.IsValid())
+										{
+											return Self->ClearButton->IsHovered() ? FStyleColors::White : FStyleColors::Foreground;
+										}
+										return FStyleColors::Foreground;
+									})
+									.Image(FAppStyle::GetBrush("Icons.X"))
+										.DesiredSizeOverride(FVector2D(12.f, 12.f))
+							]
+				]
+		]
+		]
+		];
+}
+
+#undef LOCTEXT_NAMESPACE
diff --git a/Engine/Plugins/Editor/ScriptableToolsEditorMode/Source/ScriptableToolsEditorMode/Private/UI/ScriptableToolGroupSetCustomization.cpp b/Engine/Plugins/Editor/ScriptableToolsEditorMode/Source/ScriptableToolsEditorMode/Private/UI/ScriptableToolGroupSetCustomization.cpp
new file mode 100644
--- /dev/null
+++ b/Engine/Plugins/Editor/ScriptableToolsEditorMode/Source/ScriptableToolsEditorMode/Private/UI/ScriptableToolGroupSetCustomization.cpp
@@ -0,0 +1,39 @@
+// Copyright Epic Games, Inc. All Rights Reserved.
+
+#include "UI/ScriptableToolGroupSetCustomization.h"
+
+#include "DetailWidgetRow.h"
+#include "Tags/EditableScriptableToolGroupSet.h" // IWYU pragma: keep
+#include "PropertyHandle.h"
+#include "UI/SScriptableToolGroupSetCombo.h"
+#include "Widgets/Layout/SBox.h"
+
+#define LOCTEXT_NAMESPACE "FScriptableToolTrackGroupSetCustomization"
+
+TSharedRef<IPropertyTypeCustomization> FScriptableToolGroupSetCustomization::MakeInstance()
+{
+	return MakeShareable(new FScriptableToolGroupSetCustomization);
+}
+
+void FScriptableToolGroupSetCustomization::CustomizeHeader(TSharedRef<IPropertyHandle> StructPropertyHandle, FDetailWidgetRow& HeaderRow, IPropertyTypeCustomizationUtils& CustomizationUtils)
+{
+	HeaderRow
+		.NameContent()
+		[
+			StructPropertyHandle->CreatePropertyNameWidget()
+		]
+	.ValueContent()
+		.VAlign(VAlign_Center)
+		[
+			SNew(SBox)
+			[
+				SNew(SScriptableToolGroupSetCombo)
+				.StructPropertyHandle(StructPropertyHandle)
+			]
+		];
+}
+
+void FScriptableToolGroupSetCustomization::CustomizeChildren(TSharedRef<IPropertyHandle> StructPropertyHandle, IDetailChildrenBuilder& ChildBuilder, IPropertyTypeCustomizationUtils& CustomizationUtils)
+{}
+
+#undef LOCTEXT_NAMESPACE
diff --git a/Engine/Plugins/Editor/ScriptableToolsEditorMode/Source/ScriptableToolsEditorMode/Public/ScriptableToolsEditorMode.h b/Engine/Plugins/Editor/ScriptableToolsEditorMode/Source/ScriptableToolsEditorMode/Public/ScriptableToolsEditorMode.h
--- a/Engine/Plugins/Editor/ScriptableToolsEditorMode/Source/ScriptableToolsEditorMode/Public/ScriptableToolsEditorMode.h
+++ b/Engine/Plugins/Editor/ScriptableToolsEditorMode/Source/ScriptableToolsEditorMode/Public/ScriptableToolsEditorMode.h
@@ -91,6 +91,8 @@ protected:
 
 	void InitializeModeContexts();
 
+	void RebuildScriptableToolSet();
+
 private:
 
 	TArray<TWeakObjectPtr<UScriptableToolContextObject>> ContextsToUpdateOnToolEnd;
diff --git a/Engine/Plugins/Editor/ScriptableToolsEditorMode/Source/ScriptableToolsEditorMode/Public/ScriptableToolsEditorModeSettings.h b/Engine/Plugins/Editor/ScriptableToolsEditorMode/Source/ScriptableToolsEditorMode/Public/ScriptableToolsEditorModeSettings.h
new file mode 100644
--- /dev/null
+++ b/Engine/Plugins/Editor/ScriptableToolsEditorMode/Source/ScriptableToolsEditorMode/Public/ScriptableToolsEditorModeSettings.h
@@ -0,0 +1,34 @@
+// Copyright Epic Games, Inc. All Rights Reserved.
+
+#pragma once
+
+#include "Engine/DeveloperSettings.h"
+#include "Tags/ScriptableToolGroupSet.h"
+#include "Templates/SharedPointer.h"
+
+#include "ScriptableToolsEditorModeSettings.generated.h"
+
+
+UCLASS(config = Editor)
+class SCRIPTABLETOOLSEDITORMODE_API UScriptableToolsModeCustomizationSettings : public UDeveloperSettings
+{
+	GENERATED_BODY()
+
+public:
+	// UDeveloperSettings overrides
+
+	virtual FName GetContainerName() const { return FName("Project"); }
+	virtual FName GetCategoryName() const { return FName("Plugins"); }
+	virtual FName GetSectionName() const { return FName("ScriptableTools"); }
+
+	virtual FText GetSectionText() const override;
+	virtual FText GetSectionDescription() const override;
+
+public:
+
+	UPROPERTY(config, EditAnywhere, Category = "Scriptable Tools Mode|Tool Registration", meta = (EditCondition = "!bRegisterAllTools"))
+	FScriptableToolGroupSet ToolRegistrationFilters;
+
+	bool RegisterAllTools() const {	return ToolRegistrationFilters.Groups.IsEmpty(); }
+
+};
\ No newline at end of file
diff --git a/Engine/Plugins/Editor/ScriptableToolsEditorMode/Source/ScriptableToolsEditorMode/Public/ScriptableToolsEditorModeToolkit.h b/Engine/Plugins/Editor/ScriptableToolsEditorMode/Source/ScriptableToolsEditorMode/Public/ScriptableToolsEditorModeToolkit.h
--- a/Engine/Plugins/Editor/ScriptableToolsEditorMode/Source/ScriptableToolsEditorMode/Public/ScriptableToolsEditorModeToolkit.h
+++ b/Engine/Plugins/Editor/ScriptableToolsEditorMode/Source/ScriptableToolsEditorMode/Public/ScriptableToolsEditorModeToolkit.h
@@ -41,6 +41,13 @@ public:
 	virtual void PostWarning(const FText& Message);
 	virtual void ClearWarning();
 
+	// Async Tool Loading
+	void StartAsyncToolLoading();
+	void SetAsyncProgress(float PercentLoaded);
+	void EndAsyncToolLoading();
+	bool AreToolsLoading() const;
+	TOptional<float> GetToolPercentLoaded() const;
+
 	/** Returns the Mode specific tabs in the mode toolbar **/ 
 	virtual void GetToolPaletteNames(TArray<FName>& InPaletteName) const;
 	virtual FText GetToolPaletteDisplayName(FName PaletteName) const; 
@@ -65,6 +72,15 @@ public:
 
 	virtual void ForceToolPaletteRebuild();
 
+protected:
+
+	/** FModeToolkit interface */
+	virtual void RebuildModeToolBar() override;
+	virtual bool ShouldShowModeToolbar() const override;
+
+	void RebuildModeToolPaletteWidgets();
+
+
 private:
 	const static TArray<FName> PaletteNames_Standard;
 
@@ -86,6 +102,13 @@ private:
 	TSharedPtr<SButton> CancelButton;
 	TSharedPtr<SButton> CompletedButton;
 
+	// Palette
+	bool bAsyncLoadInProgress = false;
+	float AsyncLoadProgress;
+
+	TSharedPtr<SVerticalBox> ToolBoxVBox;
+	FDelegateHandle SettingsUpdateHandle;
+
 	bool bShowRealtimeWarning = false;
 	void UpdateShowWarnings();
 
diff --git a/Engine/Plugins/Editor/ScriptableToolsEditorMode/Source/ScriptableToolsEditorMode/Public/Tags/EditableScriptableToolGroupSet.h b/Engine/Plugins/Editor/ScriptableToolsEditorMode/Source/ScriptableToolsEditorMode/Public/Tags/EditableScriptableToolGroupSet.h
new file mode 100644
--- /dev/null
+++ b/Engine/Plugins/Editor/ScriptableToolsEditorMode/Source/ScriptableToolsEditorMode/Public/Tags/EditableScriptableToolGroupSet.h
@@ -0,0 +1,28 @@
+// Copyright Epic Games, Inc. All Rights Reserved.
+
+#pragma once
+
+#include "Tags/ScriptableToolGroupSet.h"
+
+#include "EditableScriptableToolGroupSet.generated.h"
+
+UCLASS(Transient, MinimalAPI)
+class UEditableScriptableToolGroupSet : public UObject
+{
+public:
+	GENERATED_BODY()
+
+	UEditableScriptableToolGroupSet();
+
+	void SetGroups(const FScriptableToolGroupSet::FGroupSet& InGroups);
+	FScriptableToolGroupSet::FGroupSet& GetGroups();
+
+	FString GetGroupSetExportText();
+
+private:
+	UPROPERTY()
+	FScriptableToolGroupSet GroupSet;
+
+	FProperty* GroupsProperty;
+	FString GroupsPropertyAsString;
+};
\ No newline at end of file
diff --git a/Engine/Plugins/Editor/ScriptableToolsEditorMode/Source/ScriptableToolsEditorMode/Public/UI/SScriptableToolGroupSetCombo.h b/Engine/Plugins/Editor/ScriptableToolsEditorMode/Source/ScriptableToolsEditorMode/Public/UI/SScriptableToolGroupSetCombo.h
new file mode 100644
--- /dev/null
+++ b/Engine/Plugins/Editor/ScriptableToolsEditorMode/Source/ScriptableToolsEditorMode/Public/UI/SScriptableToolGroupSetCombo.h
@@ -0,0 +1,82 @@
+// Copyright Epic Games, Inc. All Rights Reserved.
+
+#pragma once
+
+#include "Input/Reply.h"
+#include "UObject/StrongObjectPtr.h"
+#include "Widgets/SCompoundWidget.h"
+#include "Tags/EditableScriptableToolGroupSet.h"
+
+template <typename ItemType> class SListView;
+
+class IPropertyHandle;
+class SMenuAnchor;
+class ITableRow;
+class STableViewBase;
+class SScriptableToolGroupSetPicker;
+class SComboButton;
+struct FScriptableToolGroupSet;
+
+/**
+ * Widget for editing a group set.
+ */
+class SScriptableToolGroupSetCombo : public SCompoundWidget
+{
+	SLATE_DECLARE_WIDGET(SScriptableToolGroupSetCombo, SCompoundWidget)
+
+public:
+
+	DECLARE_DELEGATE(FOnChanged);
+
+	SLATE_BEGIN_ARGS(SScriptableToolGroupSetCombo)
+		: _StructPropertyHandle(nullptr), _StructPtr(nullptr)
+	{}
+	SLATE_EVENT(FOnChanged, OnChanged)
+	// Used for writing changes to the group set being edited. 
+	SLATE_ARGUMENT(TSharedPtr<IPropertyHandle>, StructPropertyHandle)
+	// An alternate to the Property Handle, allowing use directly with a pointer to a tag set.
+	SLATE_ARGUMENT(FScriptableToolGroupSet*, StructPtr)
+	
+	SLATE_END_ARGS();
+
+	SCRIPTABLETOOLSEDITORMODE_API void Construct(const FArguments& InArgs);
+
+	/** Used to update internal state based on underlying data by external clients */
+	void ForceUpdate();
+
+protected:
+
+	virtual void Tick(const FGeometry& AllottedGeometry, const double InCurrentTime, const float InDeltaTime) override;
+
+private:
+	/** Returns a table row with an SScriptableToolGroupTagChip. */
+	TSharedRef<ITableRow> OnGenerateRow(UClass* InGroup, const TSharedRef<STableViewBase>& OwnerTable);
+
+	/** Instantiates the tag picker and sets it as the widget to focus for the ComboButton. */
+	TSharedRef<SWidget> OnGetMenuContent();
+
+	/** Bound to/called via SScriptableToolGroupTagChip::OnClearPressed. Removes TagToClear from group set. */
+	FReply OnClearTagClicked(UClass* InGroup);
+
+	/** Populates ActiveGroupTags with the groups currently active on the group set. */
+	void RefreshListView();
+
+	/** The set of group tags to display, based on the group tags present in the group set we are editing. */
+	TArray<UClass*> ActiveGroupTags;
+	TSharedPtr<SListView<UClass*>> ActiveGroupTagsListView;
+
+	/** Widgets we retain ownership of and refer to in a named manner. */
+	TSharedPtr<SComboButton> ComboButton;
+	TSharedPtr<SScriptableToolGroupSetPicker> TagPicker;
+
+	/** Property handle to an FScriptableToolGroupSet, used for accessing the source group set. */
+	TSharedPtr<IPropertyHandle> StructPropertyHandle;
+
+	/** A direct pointer to a group Set, an alternate pathway to the PropertyHandle */
+	FScriptableToolGroupSet* StructPtr;
+
+	/** A helper class which is used for propagating changes to the source group set. */
+	TStrongObjectPtr<UEditableScriptableToolGroupSet> HelperGroupSet;
+
+	FOnChanged OnChanged;
+};
diff --git a/Engine/Plugins/Editor/ScriptableToolsEditorMode/Source/ScriptableToolsEditorMode/Public/UI/SScriptableToolGroupSetPicker.h b/Engine/Plugins/Editor/ScriptableToolsEditorMode/Source/ScriptableToolsEditorMode/Public/UI/SScriptableToolGroupSetPicker.h
new file mode 100644
--- /dev/null
+++ b/Engine/Plugins/Editor/ScriptableToolsEditorMode/Source/ScriptableToolsEditorMode/Public/UI/SScriptableToolGroupSetPicker.h
@@ -0,0 +1,93 @@
+// Copyright Epic Games, Inc. All Rights Reserved.
+
+#pragma once
+
+#include "Containers/Map.h"
+#include "Input/Reply.h"
+#include "Layout/Visibility.h"
+#include "Widgets/SCompoundWidget.h"
+#include "UObject/Object.h"
+#include "UObject/StrongObjectPtr.h"
+
+class ITableRow;
+class SSearchBox;
+class STableViewBase;
+enum class ECheckBoxState : uint8;
+template <typename ItemType> class SListView;
+
+class UEditableScriptableToolGroupSet;
+class IPropertyHandle;
+class SComboButton;
+struct FScriptableToolGroupSet;
+
+/**
+ * Widget allowing user to edit the group tags in a given group set.
+ */
+class SScriptableToolGroupSetPicker : public SCompoundWidget
+{
+public:
+	DECLARE_DELEGATE(FOnChanged);
+
+	SLATE_BEGIN_ARGS(SScriptableToolGroupSetPicker)
+		: _StructPropertyHandle(nullptr)
+	{}
+	SLATE_EVENT(FOnChanged, OnChanged)
+	// Used for writing changes to the group set being edited. 
+	SLATE_ARGUMENT(TSharedPtr<IPropertyHandle>, StructPropertyHandle)
+	// An alternate to the Property Handle, allowing use directly with a pointer to a tag set.
+	SLATE_ARGUMENT(FScriptableToolGroupSet*, StructPtr)
+	SLATE_END_ARGS()
+
+	/** Construct the actual widget */
+		SCRIPTABLETOOLSEDITORMODE_API void Construct(const FArguments& InArgs);
+
+	/** Used to update internal state based on underlying data by external clients */
+	void ForceUpdate();
+
+private:
+	TSharedRef<SWidget> GetChildWidget();
+	TSharedRef<ITableRow> OnGenerateRow(UClass* InItem, const TSharedRef<STableViewBase>& OwnerTable);
+
+	/** Members for tracking and using search text. */
+	void OnSearchStringChanged(const FText& NewString);
+	TSharedPtr<SSearchBox> GroupSearchBox;
+	FString SearchString;
+
+	/** Populate VisibleGroupTags with the correct subset of AllGroupTags given our current search text. */
+	void RefreshListView();
+
+	/** The set of group tags to display, based on our current search text. */
+	TArray<UClass*> VisibleGroupTags;
+	TSharedPtr<SListView<UClass*>> VisibleGroupTagsListView;
+
+	/** Queries the asset registry to get all known group tags. */
+	void PopulateVisibleClasses();
+
+	/** The set of group tags which can be displayed. */
+	TArray<UClass*> AllGroupTags;
+
+	/** Functions for querying the group set for information about a particular tag. */
+	ECheckBoxState IsTagChecked(UClass* InTag);
+	void PopulateCheckedTags();
+	TMap<UClass*, bool> CheckedTags;
+
+	/** Functions for adding/removing tags to/from the group set. */
+	void OnTagChecked(UClass* NodeChecked);
+	void OnTagUnchecked(UClass* NodeUnchecked);
+	void OnUncheckAllTags();
+	void OnCreateNewTag();
+
+	/** Brings the source group set to parity with HelperGroupSet. Generally called immediately after modifying HelperGroupSet. */
+	void FlushHelperGroupSet() const;
+
+	/** Property handle to an FScriptableToolGroupSet, used for accessing the source group set. */
+	TSharedPtr<IPropertyHandle> StructPropertyHandle;
+
+	/** A direct pointer to a Group Set, an alternate pathway to the PropertyHandle */
+	FScriptableToolGroupSet* StructPtr;
+
+	/** A helper class which is used for propagating changes to the source group set. */
+	TStrongObjectPtr<UEditableScriptableToolGroupSet> HelperGroupSet;
+
+	FOnChanged OnChanged;
+};
diff --git a/Engine/Plugins/Editor/ScriptableToolsEditorMode/Source/ScriptableToolsEditorMode/Public/UI/SScriptableToolGroupTagChip.h b/Engine/Plugins/Editor/ScriptableToolsEditorMode/Source/ScriptableToolsEditorMode/Public/UI/SScriptableToolGroupTagChip.h
new file mode 100644
--- /dev/null
+++ b/Engine/Plugins/Editor/ScriptableToolsEditorMode/Source/ScriptableToolsEditorMode/Public/UI/SScriptableToolGroupTagChip.h
@@ -0,0 +1,51 @@
+// Copyright Epic Games, Inc. All Rights Reserved.
+
+#pragma once
+
+#include "Input/Reply.h"
+#include "Widgets/SCompoundWidget.h"
+
+class SButton;
+
+/**
+ * Widget for displaying a single group tag.
+ */
+class SScriptableToolGroupTagChip : public SCompoundWidget
+{
+	SLATE_DECLARE_WIDGET(SScriptableToolGroupTagChip, SCompoundWidget)
+
+public:
+
+	DECLARE_DELEGATE_RetVal(FReply, FOnClearPressed);
+
+		SLATE_BEGIN_ARGS(SScriptableToolGroupTagChip)
+	{}
+	// The group subclass associated with this chip.
+	SLATE_ARGUMENT(UClass*, TagClass)
+
+		// Callback for when clear tag button is pressed.
+		SLATE_EVENT(FOnClearPressed, OnClearPressed)
+
+		// Tooltip to display.
+		SLATE_ATTRIBUTE(FText, ToolTipText)
+
+		// Text to display.
+		SLATE_ATTRIBUTE(FText, Text)
+		SLATE_END_ARGS();
+
+	SCRIPTABLETOOLSEDITORMODE_API SScriptableToolGroupTagChip();
+	SCRIPTABLETOOLSEDITORMODE_API void Construct(const FArguments& InArgs);
+
+	/** This is public because it is used in SScriptableToolGroupSetCombo and SScriptableToolGroupSetPicker. */
+	inline static constexpr float ChipHeight = 25.f;
+
+private:
+	TSlateAttribute<FText> ToolTipTextAttribute;
+	TSlateAttribute<FText> TextAttribute;
+
+	FOnClearPressed OnClearPressed;
+
+	TSharedPtr<SButton> ClearButton;
+
+	UClass* TagClass;
+};
diff --git a/Engine/Plugins/Editor/ScriptableToolsEditorMode/Source/ScriptableToolsEditorMode/Public/UI/ScriptableToolGroupSetCustomization.h b/Engine/Plugins/Editor/ScriptableToolsEditorMode/Source/ScriptableToolsEditorMode/Public/UI/ScriptableToolGroupSetCustomization.h
new file mode 100644
--- /dev/null
+++ b/Engine/Plugins/Editor/ScriptableToolsEditorMode/Source/ScriptableToolsEditorMode/Public/UI/ScriptableToolGroupSetCustomization.h
@@ -0,0 +1,18 @@
+// Copyright Epic Games, Inc. All Rights Reserved.
+
+#pragma once
+
+#include "IPropertyTypeCustomization.h"
+
+class UScriptableToolGroupTag;
+class UEditableScriptableToolGroupSet;
+
+class FScriptableToolGroupSetCustomization : public IPropertyTypeCustomization
+{
+public:
+	static TSharedRef<IPropertyTypeCustomization> MakeInstance();
+
+	// IPropertyTypeCustomization interface
+	virtual void CustomizeHeader(TSharedRef<IPropertyHandle> StructPropertyHandle, FDetailWidgetRow& HeaderRow, IPropertyTypeCustomizationUtils& CustomizationUtils) override;
+	virtual void CustomizeChildren(TSharedRef<IPropertyHandle> StructPropertyHandle, IDetailChildrenBuilder& ChildBuilder, IPropertyTypeCustomizationUtils& CustomizationUtils) override;
+};
diff --git a/Engine/Plugins/Editor/ScriptableToolsEditorMode/Source/ScriptableToolsEditorMode/ScriptableToolsEditorMode.Build.cs b/Engine/Plugins/Editor/ScriptableToolsEditorMode/Source/ScriptableToolsEditorMode/ScriptableToolsEditorMode.Build.cs
--- a/Engine/Plugins/Editor/ScriptableToolsEditorMode/Source/ScriptableToolsEditorMode/ScriptableToolsEditorMode.Build.cs
+++ b/Engine/Plugins/Editor/ScriptableToolsEditorMode/Source/ScriptableToolsEditorMode/ScriptableToolsEditorMode.Build.cs
@@ -50,7 +50,11 @@ public class ScriptableToolsEditorMode : ModuleRules
 				"EditorScriptableToolsFramework",
 				"ModelingComponentsEditorOnly",
 
-				"ModelingComponents"
+				"ModelingComponents",
+
+				"WidgetRegistration"
+
+
 			}
 			);
 
@@ -58,7 +62,10 @@ public class ScriptableToolsEditorMode : ModuleRules
 		PrivateDependencyModuleNames.AddRange(
 			new string[]
 			{
-				"Engine"
+				"Engine",
+				"AssetRegistry",
+				"KismetWidgets"
+
 			}
 			);
 
diff --git a/Engine/Plugins/Runtime/ScriptableToolsFramework/Source/EditorScriptableToolsFramework/Private/ScriptableToolSet.cpp b/Engine/Plugins/Runtime/ScriptableToolsFramework/Source/EditorScriptableToolsFramework/Private/ScriptableToolSet.cpp
--- a/Engine/Plugins/Runtime/ScriptableToolsFramework/Source/EditorScriptableToolsFramework/Private/ScriptableToolSet.cpp
+++ b/Engine/Plugins/Runtime/ScriptableToolsFramework/Source/EditorScriptableToolsFramework/Private/ScriptableToolSet.cpp
@@ -4,19 +4,32 @@
 
 #include "ScriptableInteractiveTool.h"
 #include "ScriptableToolBuilder.h"
+#include "Tags/ScriptableToolGroupSet.h"
+#include "Engine/AssetManager.h" // Singleton access to StreamableManager
 
 
 #include "Modules/ModuleManager.h"
 #include "AssetRegistry/AssetRegistryModule.h"
 
 
-void UScriptableToolSet::ReinitializeScriptableTools()
+void UScriptableToolSet::ReinitializeScriptableTools(FPreToolsLoadedDelegate PreDelegate, FToolsLoadedDelegate PostDelegate, FToolsLoadingUpdateDelegate UpdateDelegate, FScriptableToolGroupSet* TagsToFilter)
 {
+	TRACE_CPUPROFILER_EVENT_SCOPE(TEXT("UScriptableToolSet::ReinitializeScriptableTools"))
+
+	if (bActiveLoading)
+	{
+		AsyncLoadHandle->CancelHandle();
+	}
+
+	PreDelegate.ExecuteIfBound();
+
+	bActiveLoading = true;
+	Tools.Reset();
+	ToolBuilders.Reset();
+
 	UClass* ScriptableToolClass = UScriptableInteractiveTool::StaticClass();
 	UScriptableInteractiveTool* ScriptableToolCDO = ScriptableToolClass->GetDefaultObject<UScriptableInteractiveTool>();
 
-	TSet<UClass*> PotentialToolClasses;
-
 	// Iterate over Blueprint classes to try to find UScriptableInteractiveTool blueprints
 	// Note that this code may not be fully reliable, but it appears to work so far...
 
@@ -39,6 +52,7 @@ void UScriptableToolSet::ReinitializeScriptableTools()
 	TSet<FTopLevelAssetPath> DerivedClassPaths;
 	AssetRegistry.GetDerivedClassNames(BaseModeClasses, TSet<FTopLevelAssetPath>(), DerivedClassPaths);
 
+	TArray< FSoftObjectPath > ObjectPathsToLoad;
 	for (const FTopLevelAssetPath& ClassPath : DerivedClassPaths)
 	{
 		// don't include Base tools  (note this will also catch EditorScriptableToolsFramework)
@@ -47,7 +61,50 @@ void UScriptableToolSet::ReinitializeScriptableTools()
 			continue;
 		}
 
-		FSoftObjectPath ObjectPath(ClassPath.ToString());
+		ObjectPathsToLoad.Add(ClassPath.ToString());
+	}
+	
+	TSharedPtr<FScriptableToolGroupSet> TagsToFilterCopy;
+	if (TagsToFilter)
+	{
+		TagsToFilterCopy = MakeShared<FScriptableToolGroupSet>(*TagsToFilter);
+	}
+
+	AsyncLoadHandle = UAssetManager::GetStreamableManager().RequestAsyncLoad(ObjectPathsToLoad, [this, PostDelegate, ObjectPathsToLoad, TagsToFilterCopy]() { PostToolLoad(PostDelegate, ObjectPathsToLoad, TagsToFilterCopy);  });
+
+	AsyncLoadHandle->BindUpdateDelegate(UpdateDelegate);
+}
+
+
+
+
+void UScriptableToolSet::ForEachScriptableTool(
+	TFunctionRef<void(UClass* ToolClass, UBaseScriptableToolBuilder* ToolBuilder)> ProcessToolFunc)
+{
+	if (bActiveLoading)
+	{
+		return;
+	}
+
+	for (FScriptableToolInfo& ToolInfo : Tools)
+	{
+		if (ToolInfo.ToolClass.IsValid() && ToolInfo.ToolBuilder.IsValid())
+		{
+			ProcessToolFunc(ToolInfo.ToolClass.Get(), ToolInfo.ToolBuilder.Get());
+		}
+	}
+}
+
+
+void UScriptableToolSet::PostToolLoad(FToolsLoadedDelegate Delegate, TArray< FSoftObjectPath > ObjectsLoaded, TSharedPtr<FScriptableToolGroupSet> TagsToFilter)
+{
+	TRACE_CPUPROFILER_EVENT_SCOPE(TEXT("UScriptableToolSet::PostToolLoad"))
+
+	TSet<UClass*> PotentialToolClasses;
+	UClass* ScriptableToolClass = UScriptableInteractiveTool::StaticClass();
+
+	for(FSoftObjectPath& ObjectPath : ObjectsLoaded)
+	{
 		TSoftClassPtr<UScriptableInteractiveTool> SoftClass = TSoftClassPtr<UScriptableInteractiveTool>(ObjectPath);
 		if (UClass* Class = SoftClass.LoadSynchronous())
 		{
@@ -56,13 +113,18 @@ void UScriptableToolSet::ReinitializeScriptableTools()
 				continue;
 			}
 
+			if (TagsToFilter)
+			{
+				if (!TagsToFilter->Matches(Cast<UScriptableInteractiveTool>(Class->GetDefaultObject())->GroupTags))
+				{
+					continue;
+				}
+			}
+
 			PotentialToolClasses.Add(Class);
 		}
 	}
 
-	Tools.Reset();
-	ToolBuilders.Reset();
-
 	// if the class is viable, create a ToolBuilder for it
 	for (UClass* Class : PotentialToolClasses)
 	{
@@ -86,22 +148,10 @@ void UScriptableToolSet::ReinitializeScriptableTools()
 		}
 	}
 
+	bActiveLoading = false;
+	Delegate.ExecuteIfBound();
 
-}
-
-
-
-
-void UScriptableToolSet::ForEachScriptableTool(
-	TFunctionRef<void(UClass* ToolClass, UBaseScriptableToolBuilder* ToolBuilder)> ProcessToolFunc)
-{
-	for (FScriptableToolInfo& ToolInfo : Tools)
-	{
-		if (ToolInfo.ToolClass.IsValid() && ToolInfo.ToolBuilder.IsValid())
-		{
-			ProcessToolFunc(ToolInfo.ToolClass.Get(), ToolInfo.ToolBuilder.Get());
-		}
-	}
+	AsyncLoadHandle.Reset();
 }
 
 
diff --git a/Engine/Plugins/Runtime/ScriptableToolsFramework/Source/EditorScriptableToolsFramework/Public/ScriptableToolSet.h b/Engine/Plugins/Runtime/ScriptableToolsFramework/Source/EditorScriptableToolsFramework/Public/ScriptableToolSet.h
--- a/Engine/Plugins/Runtime/ScriptableToolsFramework/Source/EditorScriptableToolsFramework/Public/ScriptableToolSet.h
+++ b/Engine/Plugins/Runtime/ScriptableToolsFramework/Source/EditorScriptableToolsFramework/Public/ScriptableToolSet.h
@@ -7,9 +7,12 @@
 
 class UScriptableInteractiveTool;
 class UBaseScriptableToolBuilder;
+struct FScriptableToolGroupSet;
 class UClass;
 
-
+DECLARE_DELEGATE(FPreToolsLoadedDelegate);
+DECLARE_DELEGATE(FToolsLoadedDelegate);
+DECLARE_DELEGATE_OneParam(FToolsLoadingUpdateDelegate, TSharedRef<struct FStreamableHandle>);
 /**
  * UScriptableToolSet represents a set of UScriptableInteractiveTool types.
  */
@@ -24,7 +27,7 @@ public:
 	 * Find all UScriptableInteractiveTool classes in the current project.
 	 * (Currently no support for filtering/etc)
 	 */
-	void ReinitializeScriptableTools();
+	void ReinitializeScriptableTools(FPreToolsLoadedDelegate PreDelegate, FToolsLoadedDelegate PostDelegate, FToolsLoadingUpdateDelegate UpdateDelegate, FScriptableToolGroupSet* TagsToFilter = nullptr);
 
 	/**
 	 * Allow external code to process each UScriptableInteractiveTool in the current ToolSet
@@ -33,6 +36,12 @@ public:
 		TFunctionRef<void(UClass* ToolClass, UBaseScriptableToolBuilder* ToolBuilder)> ProcessToolFunc);
 
 protected:
+
+	void PostToolLoad(FToolsLoadedDelegate Delegate, TArray< FSoftObjectPath > ObjectsLoaded, TSharedPtr<FScriptableToolGroupSet> TagsToFilter);
+
+	bool bActiveLoading = false;
+	TSharedPtr<FStreamableHandle> AsyncLoadHandle;
+
 	struct FScriptableToolInfo
 	{
 		TWeakObjectPtr<UClass> ToolClass = nullptr;
diff --git a/Engine/Plugins/Runtime/ScriptableToolsFramework/Source/ScriptableToolsFramework/Private/Tags/ScriptableToolGroupSet.cpp b/Engine/Plugins/Runtime/ScriptableToolsFramework/Source/ScriptableToolsFramework/Private/Tags/ScriptableToolGroupSet.cpp
new file mode 100644
--- /dev/null
+++ b/Engine/Plugins/Runtime/ScriptableToolsFramework/Source/ScriptableToolsFramework/Private/Tags/ScriptableToolGroupSet.cpp
@@ -0,0 +1,4 @@
+// Copyright Epic Games, Inc. All Rights Reserved.
+
+#include "Tags/ScriptableToolGroupSet.h"
+
diff --git a/Engine/Plugins/Runtime/ScriptableToolsFramework/Source/ScriptableToolsFramework/Private/Tags/ScriptableToolGroupTag.cpp b/Engine/Plugins/Runtime/ScriptableToolsFramework/Source/ScriptableToolsFramework/Private/Tags/ScriptableToolGroupTag.cpp
new file mode 100644
--- /dev/null
+++ b/Engine/Plugins/Runtime/ScriptableToolsFramework/Source/ScriptableToolsFramework/Private/Tags/ScriptableToolGroupTag.cpp
@@ -0,0 +1,4 @@
+// Copyright Epic Games, Inc. All Rights Reserved.
+
+#include "Tags/ScriptableToolGroupTag.h"
+
diff --git a/Engine/Plugins/Runtime/ScriptableToolsFramework/Source/ScriptableToolsFramework/Public/ScriptableInteractiveTool.h b/Engine/Plugins/Runtime/ScriptableToolsFramework/Source/ScriptableToolsFramework/Public/ScriptableInteractiveTool.h
--- a/Engine/Plugins/Runtime/ScriptableToolsFramework/Source/ScriptableToolsFramework/Public/ScriptableInteractiveTool.h
+++ b/Engine/Plugins/Runtime/ScriptableToolsFramework/Source/ScriptableToolsFramework/Public/ScriptableInteractiveTool.h
@@ -6,6 +6,8 @@
 #include "InteractiveGizmo.h"
 #include "Kismet/BlueprintFunctionLibrary.h"
 #include "ScriptableToolBuilder.h"
+#include "Tags/ScriptableToolGroupTag.h"
+#include "Tags/ScriptableToolGroupSet.h"
 #include "ScriptableInteractiveTool.generated.h"
 
 class UWorld;
@@ -892,6 +894,11 @@ public:
 	UFUNCTION(BlueprintCallable, Category = "ScriptableTool|Widgets")
 	void ClearOverlayWidget();
 
+	// Tool Tagging Support
+
+
+	UPROPERTY(EditAnywhere, Transient, Category = "Scriptable Tool Settings")
+	FScriptableToolGroupSet GroupTags;
 };
 
 
diff --git a/Engine/Plugins/Runtime/ScriptableToolsFramework/Source/ScriptableToolsFramework/Public/Tags/ScriptableToolGroupSet.h b/Engine/Plugins/Runtime/ScriptableToolsFramework/Source/ScriptableToolsFramework/Public/Tags/ScriptableToolGroupSet.h
new file mode 100644
--- /dev/null
+++ b/Engine/Plugins/Runtime/ScriptableToolsFramework/Source/ScriptableToolsFramework/Public/Tags/ScriptableToolGroupSet.h
@@ -0,0 +1,30 @@
+// Copyright Epic Games, Inc. All Rights Reserved.
+
+#pragma once
+
+#include "Containers/Map.h"
+#include "Templates/SubclassOf.h"
+#include "Tags/ScriptableToolGroupTag.h"
+
+#include "ScriptableToolGroupSet.generated.h"
+
+USTRUCT()
+struct FScriptableToolGroupSet
+{
+	GENERATED_BODY()
+
+	/**
+	 * Note: This type needs to be specified explicitly for Groups because we can't use a typedef for a UPROPERTY.
+	 */
+	typedef TSet<TSubclassOf<UScriptableToolGroupTag>> FGroupSet;
+
+	UPROPERTY(EditAnywhere, Category = "Groups")
+	TSet<TSubclassOf<UScriptableToolGroupTag>> Groups;
+
+public:
+
+	bool Matches(const FScriptableToolGroupSet& OtherSet) const
+	{
+		return !Groups.Intersect(OtherSet.Groups).IsEmpty();
+	}
+};
\ No newline at end of file
diff --git a/Engine/Plugins/Runtime/ScriptableToolsFramework/Source/ScriptableToolsFramework/Public/Tags/ScriptableToolGroupTag.h b/Engine/Plugins/Runtime/ScriptableToolsFramework/Source/ScriptableToolsFramework/Public/Tags/ScriptableToolGroupTag.h
new file mode 100644
--- /dev/null
+++ b/Engine/Plugins/Runtime/ScriptableToolsFramework/Source/ScriptableToolsFramework/Public/Tags/ScriptableToolGroupTag.h
@@ -0,0 +1,18 @@
+// Copyright Epic Games, Inc. All Rights Reserved.
+
+#pragma once
+
+#include "Misc/Optional.h"
+
+#include "ScriptableToolGroupTag.generated.h"
+
+UCLASS(Abstract, Blueprintable, Const)
+class SCRIPTABLETOOLSFRAMEWORK_API UScriptableToolGroupTag
+	: public UObject
+{
+public:
+	GENERATED_BODY()
+
+	UPROPERTY(EditAnywhere, Category = "General")
+	FString Name;
+};
