ScriptableTools: Fix issue where loading a level while scriptable tools is enabled causes a crash.

diff --git a/Engine/Plugins/Editor/ScriptableToolsEditorMode/Source/ScriptableToolsEditorMode/Private/ScriptableToolsEditorMode.cpp b/Engine/Plugins/Editor/ScriptableToolsEditorMode/Source/ScriptableToolsEditorMode/Private/ScriptableToolsEditorMode.cpp
--- a/Engine/Plugins/Editor/ScriptableToolsEditorMode/Source/ScriptableToolsEditorMode/Private/ScriptableToolsEditorMode.cpp
+++ b/Engine/Plugins/Editor/ScriptableToolsEditorMode/Source/ScriptableToolsEditorMode/Private/ScriptableToolsEditorMode.cpp
@@ -311,6 +311,7 @@ void UScriptableToolsEditorMode::OnBlueprintCompiled()
 void UScriptableToolsEditorMode::Exit()
 {
 	GEditor->OnBlueprintPreCompile().Remove(BlueprintPreCompileHandle);
+	GEditor->OnBlueprintCompiled().Remove(BlueprintCompiledHandle);
 
 	// exit any exclusive active tools w/ cancel
 	if (UInteractiveTool* ActiveTool = GetToolManager()->GetActiveTool(EToolSide::Left))
