ScriptableTools: Update icon for adding new tool groups to be more appropriate.

diff --git a/Engine/Plugins/Editor/ScriptableToolsEditorMode/Source/ScriptableToolsEditorMode/Private/UI/SScriptableToolGroupSetPicker.cpp b/Engine/Plugins/Editor/ScriptableToolsEditorMode/Source/ScriptableToolsEditorMode/Private/UI/SScriptableToolGroupSetPicker.cpp
--- a/Engine/Plugins/Editor/ScriptableToolsEditorMode/Source/ScriptableToolsEditorMode/Private/UI/SScriptableToolGroupSetPicker.cpp
+++ b/Engine/Plugins/Editor/ScriptableToolsEditorMode/Source/ScriptableToolsEditorMode/Private/UI/SScriptableToolGroupSetPicker.cpp
@@ -94,7 +94,7 @@ TSharedRef<SWidget> SScriptableToolGroupSetPicker::GetChildWidget()
 
 	MenuBuilder.AddMenuEntry(
 		LOCTEXT("ScriptableToolGroupSetPicker_CreateNewTag", "Create New Tag"),
-		FText::GetEmpty(), FSlateIcon(FAppStyle::GetAppStyleSetName(), "Icons.X"),
+		FText::GetEmpty(), FSlateIcon(FAppStyle::GetAppStyleSetName(), "Icons.Plus"),
 		FUIAction(FExecuteAction::CreateRaw(this, &SScriptableToolGroupSetPicker::OnCreateNewTag))
 	);
 
