ScriptableTools: Fixed crash when editing Scriptable Tools Mode project settings after exiting the mode.

diff --git a/Engine/Plugins/Editor/ScriptableToolsEditorMode/Source/ScriptableToolsEditorMode/Private/ScriptableToolsEditorMode.cpp b/Engine/Plugins/Editor/ScriptableToolsEditorMode/Source/ScriptableToolsEditorMode/Private/ScriptableToolsEditorMode.cpp
--- a/Engine/Plugins/Editor/ScriptableToolsEditorMode/Source/ScriptableToolsEditorMode/Private/ScriptableToolsEditorMode.cpp
+++ b/Engine/Plugins/Editor/ScriptableToolsEditorMode/Source/ScriptableToolsEditorMode/Private/ScriptableToolsEditorMode.cpp
@@ -174,7 +174,7 @@ void UScriptableToolsEditorMode::Enter()
 	ScriptableTools = NewObject<UScriptableToolSet>(this);
 
 	UScriptableToolsModeCustomizationSettings* ModeSettings = GetMutableDefault<UScriptableToolsModeCustomizationSettings>();
-	ModeSettings->OnSettingChanged().AddLambda([this](UObject*, FPropertyChangedEvent&)
+	ModeSettings->OnSettingChanged().AddWeakLambda(this, [this](UObject*, FPropertyChangedEvent&)
 	{
 		RebuildScriptableToolSet();
 	});
@@ -199,8 +199,6 @@ void UScriptableToolsEditorMode::Enter()
 
 void UScriptableToolsEditorMode::RebuildScriptableToolSet()
 {
-	UScriptableToolsModeCustomizationSettings* ModeSettings = GetMutableDefault<UScriptableToolsModeCustomizationSettings>();
-
 	auto UnregisterTools = [this]()
 	{
 		// unregister old tools from ToolManager
@@ -246,22 +244,23 @@ void UScriptableToolsEditorMode::RebuildScriptableToolSet()
 	};
 
 	// find all the Tool Blueprints
-	if (ModeSettings->RegisterAllTools())
+	if (ScriptableTools)
 	{
-		ScriptableTools->ReinitializeScriptableTools(FToolsLoadedDelegate::CreateLambda(UnregisterTools),
-													 FToolsLoadedDelegate::CreateLambda(RegisterTools),
-			                                         FToolsLoadingUpdateDelegate::CreateLambda(ToolLoadingUpdate));
+		UScriptableToolsModeCustomizationSettings* ModeSettings = GetMutableDefault<UScriptableToolsModeCustomizationSettings>();
+		if (ModeSettings->RegisterAllTools())
+		{
+			ScriptableTools->ReinitializeScriptableTools(FToolsLoadedDelegate::CreateLambda(UnregisterTools),
+														 FToolsLoadedDelegate::CreateLambda(RegisterTools),
+														 FToolsLoadingUpdateDelegate::CreateLambda(ToolLoadingUpdate));
+		}
+		else
+		{
+			ScriptableTools->ReinitializeScriptableTools(FToolsLoadedDelegate::CreateLambda(UnregisterTools),
+														 FToolsLoadedDelegate::CreateLambda(RegisterTools),
+														 FToolsLoadingUpdateDelegate::CreateLambda(ToolLoadingUpdate),
+														 &ModeSettings->ToolRegistrationFilters);
+		}
 	}
-	else
-	{
-		ScriptableTools->ReinitializeScriptableTools(FToolsLoadedDelegate::CreateLambda(UnregisterTools),
-												     FToolsLoadedDelegate::CreateLambda(RegisterTools),
-			                                         FToolsLoadingUpdateDelegate::CreateLambda(ToolLoadingUpdate),
-			                                         &ModeSettings->ToolRegistrationFilters);
-	}	
-
-
-
 }
 
 void UScriptableToolsEditorMode::InitializeModeContexts()
@@ -310,6 +309,9 @@ void UScriptableToolsEditorMode::OnBlueprintCompiled()
 
 void UScriptableToolsEditorMode::Exit()
 {
+	UScriptableToolsModeCustomizationSettings* ModeSettings = GetMutableDefault<UScriptableToolsModeCustomizationSettings>();
+	ModeSettings->OnSettingChanged().RemoveAll(this);
+	
 	GEditor->OnBlueprintPreCompile().Remove(BlueprintPreCompileHandle);
 	GEditor->OnBlueprintCompiled().Remove(BlueprintCompiledHandle);
 
