ScriptableTools: Add support for dynamic command registration.

diff --git a/Engine/Plugins/Editor/ScriptableToolsEditorMode/Source/ScriptableToolsEditorMode/Private/ScriptableToolsEditorModeManagerCommands.cpp b/Engine/Plugins/Editor/ScriptableToolsEditorMode/Source/ScriptableToolsEditorMode/Private/ScriptableToolsEditorModeManagerCommands.cpp
--- a/Engine/Plugins/Editor/ScriptableToolsEditorMode/Source/ScriptableToolsEditorMode/Private/ScriptableToolsEditorModeManagerCommands.cpp
+++ b/Engine/Plugins/Editor/ScriptableToolsEditorMode/Source/ScriptableToolsEditorMode/Private/ScriptableToolsEditorModeManagerCommands.cpp
@@ -35,6 +35,36 @@ TSharedPtr<FUICommandInfo> FScriptableToolsEditorModeManagerCommands::FindToolBy
 	return TSharedPtr<FUICommandInfo>();
 }
 
+void FScriptableToolsEditorModeManagerCommands::NotifyCommandsChanged() const
+{
+	if (TSharedPtr<FScriptableToolsEditorModeManagerCommands> Commands = Instance.Pin())
+	{
+		CommandsChanged.Broadcast(*Commands);
+	}
+}
+
+TSharedPtr<FUICommandInfo> FScriptableToolsEditorModeManagerCommands::RegisterCommand(
+		FName InCommandName,
+		const FText& InCommandLabel,
+		const FText& InCommandTooltip,
+		const FSlateIcon& InIcon,
+		const EUserInterfaceActionType InUserInterfaceType,
+		const FInputChord& InDefaultChord)
+{
+	const TSharedRef<FBindingContext> SharedThis = AsShared();
+	TSharedPtr<FUICommandInfo> Result;
+	FUICommandInfo::MakeCommandInfo(
+		SharedThis,
+		Result,
+		InCommandName,
+		InCommandLabel,
+		InCommandTooltip,
+		InIcon,
+		InUserInterfaceType,
+		InDefaultChord);
+	RegisteredTools.Add(FStartToolCommand{ InCommandName.ToString(), Result});
+	return Result;
+}
 
 void FScriptableToolsEditorModeManagerCommands::RegisterCommands()
 {
@@ -43,12 +73,6 @@ void FScriptableToolsEditorModeManagerCommands::RegisterCommands()
 		UI_COMMAND(ToolCommandInfo, ToolName, ToolTip, EUserInterfaceActionType::ToggleButton, FInputChord()); \
 		RegisteredTools.Add(FStartToolCommand{ ToolName, ToolCommandInfo });
 
-	REGISTER_TOOL_COMMAND(BeginMeshInspectorTool, "Inspct", "Inspect mesh attributes");
-	REGISTER_TOOL_COMMAND(BeginPhysicsInspectorTool, "PInspct", "Inspect physics geometry for selected meshes");
-	REGISTER_TOOL_COMMAND(BeginSetCollisionGeometryTool, "Msh2Coll", "Convert selected meshes to Simple Collision Geometry (for last selected)");
-
-
-
 	UI_COMMAND(AcceptActiveTool, "Accept", "Accept the active tool", EUserInterfaceActionType::Button, FInputChord());
 	UI_COMMAND(CancelActiveTool, "Cancel", "Cancel the active tool", EUserInterfaceActionType::Button, FInputChord());
 	UI_COMMAND(CompleteActiveTool, "Done", "Complete the active tool", EUserInterfaceActionType::Button, FInputChord());
diff --git a/Engine/Plugins/Editor/ScriptableToolsEditorMode/Source/ScriptableToolsEditorMode/Public/ScriptableToolsEditorModeManagerCommands.h b/Engine/Plugins/Editor/ScriptableToolsEditorMode/Source/ScriptableToolsEditorMode/Public/ScriptableToolsEditorModeManagerCommands.h
--- a/Engine/Plugins/Editor/ScriptableToolsEditorMode/Source/ScriptableToolsEditorMode/Public/ScriptableToolsEditorModeManagerCommands.h
+++ b/Engine/Plugins/Editor/ScriptableToolsEditorMode/Source/ScriptableToolsEditorMode/Public/ScriptableToolsEditorModeManagerCommands.h
@@ -26,18 +26,17 @@ public:
 	 * Find Tool start-command below by registered name (tool icon name in Mode palette)
 	 */
 	TSharedPtr<FUICommandInfo> FindToolByName(FString Name, bool& bFound) const;
-	
-	//
-	// These commands are set up to launch registered Tools via the ToolManager in ScriptableToolsEditorMode.cpp
-	//
-	TSharedPtr<FUICommandInfo> BeginMeshInspectorTool;
-	TSharedPtr<FUICommandInfo> BeginPhysicsInspectorTool;
-	TSharedPtr<FUICommandInfo> BeginSetCollisionGeometryTool;
-
-
-	TSharedPtr<FUICommandInfo> BeginMeshToPropTool;
 
+	TSharedPtr<FUICommandInfo> RegisterCommand(
+		FName InCommandName,
+		const FText& InCommandLabel,
+		const FText& InCommandTooltip,
+		const FSlateIcon& InIcon,
+		const EUserInterfaceActionType InUserInterfaceType,
+		const FInputChord& InDefaultChord);
 
+	void NotifyCommandsChanged() const;
+	
 	//
 	// Accept/Cancel/Complete commands are used to end the active Tool via ToolManager
 	//
