ScriptableTools: Added new ToolkitBuilder UI layout as default.

The older layout is still accessible through the PerProject settings bUseLegacyPalette option.

diff --git a/Engine/Plugins/Editor/ScriptableToolsEditorMode/Source/ScriptableToolsEditorMode/Private/ScriptableToolsEditorMode.cpp b/Engine/Plugins/Editor/ScriptableToolsEditorMode/Source/ScriptableToolsEditorMode/Private/ScriptableToolsEditorMode.cpp
--- a/Engine/Plugins/Editor/ScriptableToolsEditorMode/Source/ScriptableToolsEditorMode/Private/ScriptableToolsEditorMode.cpp
+++ b/Engine/Plugins/Editor/ScriptableToolsEditorMode/Source/ScriptableToolsEditorMode/Private/ScriptableToolsEditorMode.cpp
@@ -19,6 +19,7 @@
 
 #include "ScriptableToolBuilder.h"
 #include "ScriptableToolSet.h"
+#include "ScriptableInteractiveTool.h"
 
 #include "InteractiveToolQueryInterfaces.h" // IInteractiveToolExclusiveToolAPI
 #include "ToolContextInterfaces.h"
@@ -204,11 +205,13 @@ void UScriptableToolsEditorMode::RebuildScriptableToolSet()
 		// unregister old tools from ToolManager
 		ScriptableTools->ForEachScriptableTool([&](UClass* ToolClass, UInteractiveToolBuilder* ToolBuilder)
 			{
-				FString UseName;
-				ToolClass->GetClassPathName().ToString(UseName);
-				GetToolManager(EToolsContextScope::EdMode)->UnregisterToolType(UseName);
+				FString ToolIdentifier;
+				ToolClass->GetClassPathName().ToString(ToolIdentifier);
+				GetToolManager(EToolsContextScope::EdMode)->UnregisterToolType(ToolIdentifier);
 			});
 
+		FScriptableToolsEditorModeManagerCommands::Unregister();
+		
 		if (Toolkit.IsValid())
 		{
 			FScriptableToolsEditorModeToolkit* ModeToolkit = (FScriptableToolsEditorModeToolkit*)Toolkit.Get();
@@ -218,18 +221,96 @@ void UScriptableToolsEditorMode::RebuildScriptableToolSet()
 
 	auto RegisterTools = [this]()
 	{
+		FScriptableToolsEditorModeManagerCommands::Register();
+		
+		FScriptableToolsEditorModeManagerCommands& ToolManagerCommands = const_cast<FScriptableToolsEditorModeManagerCommands&>(FScriptableToolsEditorModeManagerCommands::Get());
 		// register each of them with ToolManager
 		ScriptableTools->ForEachScriptableTool([&](UClass* ToolClass, UInteractiveToolBuilder* ToolBuilder)
 		{
-			FString UseName;
-			ToolClass->GetClassPathName().ToString(UseName);
-			GetToolManager(EToolsContextScope::EdMode)->RegisterToolType(UseName, ToolBuilder);
+			FString ToolIdentifier;
+			ToolClass->GetClassPathName().ToString(ToolIdentifier);
+			GetToolManager(EToolsContextScope::EdMode)->RegisterToolType(ToolIdentifier, ToolBuilder);
+
+			// Register commands for each tool
+			UScriptableInteractiveTool* ToolCDO = Cast<UScriptableInteractiveTool>(ToolClass->GetDefaultObject());
+			if (ToolCDO)
+			{
+				FString ToolNameString = ToolCDO->ToolName.IsEmpty() ? ToolIdentifier : ToolCDO->ToolName.ToString();
+				const FName ToolName(ToolNameString);
+				const FText ToolLabel = ToolCDO->ToolLongName.IsEmpty() ? FText::FromString(ToolNameString) : ToolCDO->ToolLongName;
+				const FText ToolTooltip = ToolCDO->ToolTooltip.IsEmpty() ? ToolLabel : ToolCDO->ToolTooltip;
+				TSharedPtr<FUICommandInfo> ToolCommand = ToolManagerCommands.RegisterCommand(
+					FName(ToolIdentifier),
+					ToolLabel,
+					ToolTooltip,
+					FSlateIcon(),
+					EUserInterfaceActionType::ToggleButton,
+					FInputChord());
+
+				if (Toolkit.IsValid())
+				{
+					const TSharedRef<FUICommandList>& CommandList = Toolkit->GetToolkitCommands();
+					CommandList->MapAction(ToolCommand,
+					FExecuteAction::CreateLambda([this, ToolClass, ToolIdentifier]()
+					{
+						if (GetToolManager()->SelectActiveToolType(EToolSide::Mouse, ToolIdentifier))
+						{
+							if (GetToolManager()->CanActivateTool(EToolSide::Mouse, ToolIdentifier)) 
+							{
+								GetToolManager()->ActivateTool(EToolSide::Mouse);
+							}
+						}
+						else
+						{
+							UE_LOG(LogTemp, Warning, TEXT("FAILED TO SET ACTIVE TOOL TYPE!"));
+						}
+
+					}),
+					FCanExecuteAction::CreateLambda([this, ToolClass, ToolIdentifier]()
+					{
+						if (GetToolManager()->SelectActiveToolType(EToolSide::Mouse, ToolIdentifier))
+						{
+							return GetToolManager()->CanActivateTool(EToolSide::Mouse, ToolIdentifier);
+						}
+						else
+						{
+							UE_LOG(LogTemp, Warning, TEXT("FAILED TO SET ACTIVE TOOL TYPE!"));
+							return false;
+						}
+					}),
+					FIsActionChecked::CreateLambda([this, ToolClass, ToolIdentifier]()
+					{
+						return GetToolManager()->GetActiveToolName(EToolSide::Mouse) == ToolIdentifier;
+					}));
+				}
+			}
 		});
 
 		if (Toolkit.IsValid())
 		{
 			FScriptableToolsEditorModeToolkit* ModeToolkit = (FScriptableToolsEditorModeToolkit*)Toolkit.Get();
 			ModeToolkit->EndAsyncToolLoading();
+
+			// Register LoadPalette commands
+			TArray<FName> PaletteNames;
+			ModeToolkit->GetActiveToolPaletteNames(PaletteNames);
+			for (const FName PaletteName : PaletteNames)
+			{
+				const FString PaletteNameString = PaletteName.ToString();
+				ensure(!PaletteNameString.IsEmpty());
+				const FString LoadPaletteString = "LoadPalette" + PaletteNameString;
+				const FText PaletteNameText = FText::FromString(PaletteNameString);
+				ToolManagerCommands.RegisterCommand(
+					FName(LoadPaletteString),
+					PaletteNameText,
+					PaletteNameText,
+					FSlateIcon(),
+					EUserInterfaceActionType::ToggleButton,
+					FInputChord());
+			}
+
+			ToolManagerCommands.NotifyCommandsChanged();
+
 			ModeToolkit->ForceToolPaletteRebuild();
 		}
 	};
diff --git a/Engine/Plugins/Editor/ScriptableToolsEditorMode/Source/ScriptableToolsEditorMode/Private/ScriptableToolsEditorModeToolkit.cpp b/Engine/Plugins/Editor/ScriptableToolsEditorMode/Source/ScriptableToolsEditorMode/Private/ScriptableToolsEditorModeToolkit.cpp
--- a/Engine/Plugins/Editor/ScriptableToolsEditorMode/Source/ScriptableToolsEditorMode/Private/ScriptableToolsEditorModeToolkit.cpp
+++ b/Engine/Plugins/Editor/ScriptableToolsEditorMode/Source/ScriptableToolsEditorMode/Private/ScriptableToolsEditorModeToolkit.cpp
@@ -362,19 +362,52 @@ void FScriptableToolsEditorModeToolkit::CustomizeModeDetailsViewArgs(FDetailsVie
 	//ArgsInOut.ColumnWidth = 0.3f;
 }
 
+void FScriptableToolsEditorModeToolkit::RegisterPalettes()
+{
+	UEdMode* ScriptableMode = GetScriptableEditorMode().Get();
+
+	ToolkitSections = MakeShared<FToolkitSections>();
+	FToolkitBuilderArgs ToolkitBuilderArgs(ScriptableMode->GetModeInfo().ToolbarCustomizationName);
+	ToolkitBuilderArgs.ToolkitCommandList = GetToolkitCommands();
+	ToolkitBuilderArgs.ToolkitSections = ToolkitSections;
+	ToolkitBuilderArgs.SelectedCategoryTitleVisibility = EVisibility::Collapsed;
+	// This lets us re-show the buttons if the user clicks a category with a tool still active.
+	ToolkitBuilderArgs.CategoryReclickBehavior = FToolkitBuilder::ECategoryReclickBehavior::TreatAsChanged;
+	ToolkitBuilder = MakeShared<FToolkitBuilder>(ToolkitBuilderArgs);
+	ToolkitBuilder->SetCategoryButtonLabelVisibility(true);
+
+	// Actual population of palettes is deferred to support async load of
+	// scriptable tools. Proper population should happen as part of
+	// ForceToolPaletteRebuild.
+	
+	// if selected palette changes, make sure we are showing the palette command buttons, which may be hidden by active Tool
+	ActivePaletteChangedHandle = ToolkitBuilder->OnActivePaletteChanged.AddLambda([this]()
+	{
+		ToolkitBuilder->SetActivePaletteCommandsVisibility(EVisibility::Visible);
+	});
+}
+
 void FScriptableToolsEditorModeToolkit::Init(const TSharedPtr<IToolkitHost>& InitToolkitHost, TWeakObjectPtr<UEdMode> InOwningMode)
 {
+	const UScriptableToolsModeCustomizationSettings* UISettings = GetDefault<UScriptableToolsModeCustomizationSettings>();
+	bUsesToolkitBuilder = !UISettings->bUseLegacyPalette;
+
 	// Have to create the ToolkitWidget here because FModeToolkit::Init() is going to ask for it and add
 	// it to the Mode panel, and not ask again afterwards. However we have to call Init() to get the 
 	// ModeDetailsView created, that we need to add to the ToolkitWidget. So, we will create the Widget
 	// here but only add the rows to it after we call Init()
+	
 	TSharedPtr<SVerticalBox> ToolkitWidgetVBox = SNew(SVerticalBox);
-	SAssignNew(ToolkitWidget, SBorder)
-		.HAlign(HAlign_Fill)
-		.Padding(4)
-		[
-			ToolkitWidgetVBox->AsShared()
-		];
+	
+	if ( !bUsesToolkitBuilder )
+	{
+		SAssignNew(ToolkitWidget, SBorder)
+			.HAlign(HAlign_Fill)
+			.Padding(4)
+			[
+				ToolkitWidgetVBox->AsShared()
+			];
+	}
 
 	FModeToolkit::Init(InitToolkitHost, InOwningMode);
 
@@ -400,27 +433,80 @@ void FScriptableToolsEditorModeToolkit::Init(const TSharedPtr<IToolkitHost>& Ini
 		.ColorAndOpacity(FSlateColor(FLinearColor(0.9f, 0.15f, 0.15f)));
 	ToolWarningArea->SetText(FText::GetEmpty());
 
-	// add the various sections to the mode toolbox
-	ToolkitWidgetVBox->AddSlot().AutoHeight().HAlign(HAlign_Fill).Padding(5)
+	ToolPaletteTagPanel = SNew(SHorizontalBox)
+		+ SHorizontalBox::Slot()
+		.AutoWidth()
+		.FillWidth(1.f)
+		.Padding(1.f)
 		[
-			ModeWarningArea->AsShared()
+			SNew(SToolPaletteTagPanel, this)
 		];
-	ToolkitWidgetVBox->AddSlot().AutoHeight().HAlign(HAlign_Fill).Padding(5)
+
+	ToolPaletteLoadBar = SNew(SHorizontalBox)
+		+ SHorizontalBox::Slot()
+		.AutoWidth()
+		.FillWidth(1.f)
+		.Padding(1.f)
 		[
-			ModeHeaderArea->AsShared()
+			SNew(SToolPaletteLoadBar, this)
 		];
-	ToolkitWidgetVBox->AddSlot().AutoHeight().HAlign(HAlign_Fill).Padding(5)
+
+	ToolPaletteHeader = SNew(SVerticalBox)
+		+ SVerticalBox::Slot()
+		.AutoHeight()
 		[
-			ToolWarningArea->AsShared()
-		];
-	ToolkitWidgetVBox->AddSlot().HAlign(HAlign_Fill).FillHeight(1.f)
+			ToolPaletteTagPanel->AsShared()
+		]
+		+ SVerticalBox::Slot()
+		.AutoHeight()
 		[
-			ModeDetailsView->AsShared()
+			ToolPaletteLoadBar->AsShared()
 		];
 
+	if (bUsesToolkitBuilder)
+	{
+		RegisterPalettes();
+	}
+	else
+	{
+		// add the various sections to the mode toolbox
+		ToolkitWidgetVBox->AddSlot().AutoHeight().HAlign(HAlign_Fill).Padding(5)
+			[
+				ModeWarningArea->AsShared()
+			];
+		ToolkitWidgetVBox->AddSlot().AutoHeight().HAlign(HAlign_Fill).Padding(5)
+			[
+				ModeHeaderArea->AsShared()
+			];
+		ToolkitWidgetVBox->AddSlot().AutoHeight().HAlign(HAlign_Fill).Padding(5)
+			[
+				ToolWarningArea->AsShared()
+			];
+		ToolkitWidgetVBox->AddSlot().HAlign(HAlign_Fill).FillHeight(1.f)
+			[
+				ModeDetailsView->AsShared()
+			];
+	}
+
 	ClearNotification();
 	ClearWarning();
 
+	if (HasToolkitBuilder())
+	{
+		ToolkitSections->ModeWarningArea = ModeWarningArea;
+		ToolkitSections->DetailsView = ModeDetailsView;
+		ToolkitSections->ToolWarningArea = ToolWarningArea;
+		ToolkitSections->Footer = ToolPaletteHeader;
+
+		SAssignNew(ToolkitWidget, SBorder)
+		.HAlign(HAlign_Fill)
+		.Padding(0)
+		.BorderImage(FAppStyle::GetBrush("ToolPanel.GroupBorder"))
+		[
+			ToolkitBuilder->GenerateWidget()->AsShared()
+		];	
+	}
+
 	ActiveToolName = FText::GetEmpty();
 	ActiveToolMessage = FText::GetEmpty();
 
@@ -529,7 +615,12 @@ void FScriptableToolsEditorModeToolkit::UpdateActiveToolCategories()
 			FName CategoryName(ToolCDO->ToolCategory.ToString());
 			if (ActiveToolCategories.Contains(CategoryName) == false)
 			{
-				ActiveToolCategories.Add(CategoryName, ToolCDO->ToolCategory);
+				ActiveToolCategories.Add(CategoryName, {});
+			}
+
+			if (TArray<FScriptableToolData>* CategoryTools = ActiveToolCategories.Find(CategoryName))
+			{
+				CategoryTools->Add(FScriptableToolData{ToolCDO->ToolCategory, ToolClass, ToolBuilder});
 			}
 		}
 	});
@@ -646,16 +737,17 @@ void FScriptableToolsEditorModeToolkit::GetToolPaletteNames(TArray<FName>& Palet
 
 
 FText FScriptableToolsEditorModeToolkit::GetToolPaletteDisplayName(FName Palette) const
-{ 
-	if (ActiveToolCategories.Contains(Palette))
-	{
-		return ActiveToolCategories[Palette];
-	}
+{
 	return FText::FromName(Palette);
 }
 
 void FScriptableToolsEditorModeToolkit::BuildToolPalette(FName PaletteIndex, class FToolBarBuilder& ToolbarBuilder) 
 {
+	if (HasToolkitBuilder())
+	{
+		return;
+	}
+	
 	const FScriptableToolsEditorModeManagerCommands& Commands = FScriptableToolsEditorModeManagerCommands::Get();
 
 	static TArray<TSharedPtr<FUIAction>> ActionsHack;
@@ -780,6 +872,12 @@ void FScriptableToolsEditorModeToolkit::BuildToolPalette(FName PaletteIndex, cla
 void FScriptableToolsEditorModeToolkit::InvokeUI()
 {
 	FModeToolkit::InvokeUI();
+
+	// FModeToolkit::UpdatePrimaryModePanel() wrapped our GetInlineContent() output in a SScrollBar widget,
+	// however this doesn't make sense as we want to dock panels to the "top" and "bottom" of our mode panel area,
+	// and the details panel in the middle has it's own scrollbar already. The SScrollBar is hardcoded as the content
+	// of FModeToolkit::InlineContentHolder so we can just replace it here
+	InlineContentHolder->SetContent(GetInlineContent().ToSharedRef());
 }
 
 
@@ -788,6 +886,10 @@ void FScriptableToolsEditorModeToolkit::ForceToolPaletteRebuild()
 {
 	this->UpdateActiveToolCategories();
 
+	if (HasToolkitBuilder())
+	{
+		RebuildModeToolkitBuilderPalettes();
+	}
 	if (ModeUILayer.IsValid() && HasIntegratedToolPalettes() == false)
 	{
 		if (TSharedPtr<FAssetEditorModeUILayer> ModeUILayerPtr = ModeUILayer.Pin())
@@ -811,16 +913,18 @@ void FScriptableToolsEditorModeToolkit::ForceToolPaletteRebuild()
 			}
 		}
 	}
-
 }
 
+void FScriptableToolsEditorModeToolkit::GetActiveToolPaletteNames(TArray<FName>& OutPaletteNames)
+{
+	UpdateActiveToolCategories();
+	GetToolPaletteNames(OutPaletteNames);
+}
 
 void FScriptableToolsEditorModeToolkit::OnToolPaletteChanged(FName PaletteName) 
 {
 }
 
-
-
 void FScriptableToolsEditorModeToolkit::EnableShowRealtimeWarning(bool bEnable)
 {
 	if (bShowRealtimeWarning != bEnable)
@@ -853,6 +957,18 @@ void FScriptableToolsEditorModeToolkit::OnToolStarted(UInteractiveToolManager* M
 		}
 	}
 
+	if (HasToolkitBuilder())
+	{
+		ToolkitBuilder->SetActiveToolDisplayName(ActiveToolName);
+		if (const UScriptableToolsModeCustomizationSettings* Settings = GetDefault<UScriptableToolsModeCustomizationSettings>())
+		{
+			if (Settings->bAlwaysShowToolButtons == false)
+			{
+				ToolkitBuilder->SetActivePaletteCommandsVisibility(EVisibility::Collapsed);
+			}
+		}
+	}
+
 	// try to update icon
 	FString ActiveToolIdentifier = GetScriptableEditorMode()->GetToolManager(EToolsContextScope::EdMode)->GetActiveToolName(EToolSide::Left);
 	ActiveToolIdentifier.InsertAt(0, ".");
@@ -871,6 +987,12 @@ void FScriptableToolsEditorModeToolkit::OnToolEnded(UInteractiveToolManager* Man
 
 	ModeDetailsView->SetObject(nullptr);
 	ActiveToolName = FText::GetEmpty();
+	if (HasToolkitBuilder())
+	{
+		ToolkitBuilder->SetActiveToolDisplayName(FText::GetEmpty());
+		ToolkitBuilder->SetActivePaletteCommandsVisibility(EVisibility::Visible);
+	}
+	
 	ModeHeaderArea->SetVisibility(EVisibility::Visible);
 	ModeHeaderArea->SetText(LOCTEXT("SelectToolLabel", "Select a Tool from the Toolbar"));
 	ClearNotification();
@@ -1053,4 +1175,53 @@ void FScriptableToolsEditorModeToolkit::RebuildModeToolPaletteWidgets()
 
 }
 
+void FScriptableToolsEditorModeToolkit::RebuildModeToolkitBuilderPalettes()
+{
+	const FScriptableToolsEditorModeManagerCommands& Commands = FScriptableToolsEditorModeManagerCommands::Get();
+	const TSharedPtr<FUICommandList> CommandList = GetToolkitCommands();
+	TMap<FName, TArray<TSharedPtr<FUICommandInfo>>> CommandLists = OwningEditorMode->GetModeCommands();
+
+	TSharedPtr<FUICommandInfo> ActivePaletteOnLoad;
+	TArray<FName> PaletteNames;
+	GetToolPaletteNames(PaletteNames);
+	const int PaletteCount = PaletteNames.Num();
+	for (int32 PaletteId = 0; PaletteId < PaletteCount; ++PaletteId)
+	{
+		const FName PaletteName = PaletteNames[PaletteId];
+		if (const TArray<FScriptableToolData>* PaletteTools = ActiveToolCategories.Find(PaletteName))
+		{
+			TArray<TSharedPtr<FUICommandInfo>> PaletteCommandList;
+			for (const FScriptableToolData& Tool : *PaletteTools)
+			{
+				UScriptableInteractiveTool* ToolCDO = Cast<UScriptableInteractiveTool>(Tool.ToolClass->GetDefaultObject());
+				if (!ToolCDO)
+				{
+					continue;
+				}
+
+				bool bFound = false;
+				TSharedPtr<FUICommandInfo> ToolCommand = Commands.FindToolByName(ToolCDO->ToolName.ToString(), bFound);
+				if (bFound)
+				{
+					PaletteCommandList.Add(ToolCommand);
+				}
+			}
+
+			bool bFound = false;
+			const FString LoadPaletteName = "LoadPalette" + PaletteName.ToString();
+			TSharedPtr<FUICommandInfo> LoadPaletteCommand = Commands.FindToolByName(LoadPaletteName, bFound);
+			ActivePaletteOnLoad = LoadPaletteCommand;
+			ensure(bFound);
+			ToolkitBuilder->AddPalette(
+				MakeShareable( new FToolPalette(LoadPaletteCommand.ToSharedRef(), PaletteCommandList)) );
+		}
+	}
+
+	if (ActivePaletteOnLoad)
+	{
+		ToolkitBuilder->SetActivePaletteOnLoad(ActivePaletteOnLoad.Get());
+	}
+	ToolkitBuilder->RefreshCategoryToolbarWidget(false);
+}
+
 #undef LOCTEXT_NAMESPACE
diff --git a/Engine/Plugins/Editor/ScriptableToolsEditorMode/Source/ScriptableToolsEditorMode/Public/ScriptableToolsEditorModeSettings.h b/Engine/Plugins/Editor/ScriptableToolsEditorMode/Source/ScriptableToolsEditorMode/Public/ScriptableToolsEditorModeSettings.h
--- a/Engine/Plugins/Editor/ScriptableToolsEditorMode/Source/ScriptableToolsEditorMode/Public/ScriptableToolsEditorModeSettings.h
+++ b/Engine/Plugins/Editor/ScriptableToolsEditorMode/Source/ScriptableToolsEditorMode/Public/ScriptableToolsEditorModeSettings.h
@@ -29,6 +29,14 @@ public:
 	UPROPERTY(config, EditAnywhere, Category = "Scriptable Tools Mode|Tool Registration", meta = (EditCondition = "!bRegisterAllTools"))
 	FScriptableToolGroupSet ToolRegistrationFilters;
 
+	/** Toggle between the Legacy Scriptable Tools Palette and the new UI (requires exiting and re-entering the Mode) */
+	UPROPERTY(config, EditAnywhere, Category = "Scriptable Tools Mode|UI Customization")
+	bool bUseLegacyPalette = false;
+
+	/** If true, Tool buttons will always be shown when in a Tool. By default they will be hidden. */
+	UPROPERTY(config, EditAnywhere, Category = "Scriptable Tools Mode|UI Customization")
+	bool bAlwaysShowToolButtons = false;
+
 	bool RegisterAllTools() const {	return ToolRegistrationFilters.GetGroups().IsEmpty(); }
 
 };
\ No newline at end of file
diff --git a/Engine/Plugins/Editor/ScriptableToolsEditorMode/Source/ScriptableToolsEditorMode/Public/ScriptableToolsEditorModeToolkit.h b/Engine/Plugins/Editor/ScriptableToolsEditorMode/Source/ScriptableToolsEditorMode/Public/ScriptableToolsEditorModeToolkit.h
--- a/Engine/Plugins/Editor/ScriptableToolsEditorMode/Source/ScriptableToolsEditorMode/Public/ScriptableToolsEditorModeToolkit.h
+++ b/Engine/Plugins/Editor/ScriptableToolsEditorMode/Source/ScriptableToolsEditorMode/Public/ScriptableToolsEditorModeToolkit.h
@@ -14,6 +14,7 @@ class IDetailsView;
 class SButton;
 class STextBlock;
 class UBlueprint;
+class UBaseScriptableToolBuilder;
 
 class SCRIPTABLETOOLSEDITORMODE_API FScriptableToolsEditorModeToolkit : public FModeToolkit
 {
@@ -49,12 +50,12 @@ public:
 	TOptional<float> GetToolPercentLoaded() const;
 
 	/** Returns the Mode specific tabs in the mode toolbar **/ 
-	virtual void GetToolPaletteNames(TArray<FName>& InPaletteName) const;
-	virtual FText GetToolPaletteDisplayName(FName PaletteName) const; 
-	virtual void BuildToolPalette(FName PaletteName, class FToolBarBuilder& ToolbarBuilder);
+	virtual void GetToolPaletteNames(TArray<FName>& InPaletteName) const override;
+	virtual FText GetToolPaletteDisplayName(FName PaletteName) const override; 
+	virtual void BuildToolPalette(FName PaletteName, class FToolBarBuilder& ToolbarBuilder) override;
 	virtual void OnToolPaletteChanged(FName PaletteName) override;
-	virtual bool HasIntegratedToolPalettes() const { return false; }
-	virtual bool HasExclusiveToolPalettes() const { return false; }
+	virtual bool HasIntegratedToolPalettes() const override { return false; }
+	virtual bool HasExclusiveToolPalettes() const override { return false; }
 
 	virtual FText GetActiveToolDisplayName() const override { return ActiveToolName; }
 	virtual FText GetActiveToolMessage() const override { return ActiveToolMessage; }
@@ -72,6 +73,8 @@ public:
 
 	virtual void ForceToolPaletteRebuild();
 
+	void GetActiveToolPaletteNames(TArray<FName>& OutPaletteNames);
+
 protected:
 
 	/** FModeToolkit interface */
@@ -79,6 +82,7 @@ protected:
 	virtual bool ShouldShowModeToolbar() const override;
 
 	void RebuildModeToolPaletteWidgets();
+	void RebuildModeToolkitBuilderPalettes();
 
 
 private:
@@ -92,6 +96,9 @@ private:
 	TSharedPtr<SWidget> ToolkitWidget;
 	void UpdateActiveToolProperties();
 	void InvalidateCachedDetailPanelState(UObject* ChangedObject);
+	
+	void RegisterPalettes();
+	FDelegateHandle ActivePaletteChangedHandle;
 
 	TSharedPtr<SWidget> ViewportOverlayWidget;
 
@@ -109,10 +116,20 @@ private:
 	TSharedPtr<SVerticalBox> ToolBoxVBox;
 	FDelegateHandle SettingsUpdateHandle;
 
+	TSharedPtr<SWidget> ToolPaletteHeader;
+	TSharedPtr<SWidget> ToolPaletteTagPanel;
+	TSharedPtr<SWidget> ToolPaletteLoadBar;
+
 	bool bShowRealtimeWarning = false;
 	void UpdateShowWarnings();
 
-	TMap<FName, FText> ActiveToolCategories;
+	struct FScriptableToolData
+	{
+		FText Category;
+		UClass* ToolClass = nullptr;
+		UBaseScriptableToolBuilder* Builder = nullptr;
+	};
+	TMap<FName, TArray<FScriptableToolData>> ActiveToolCategories;
 	void UpdateActiveToolCategories();
 
 	bool bFirstInitializeAfterModeSetup = true;
