ScriptableTools: Set "AlwaysShowToolButtons" to true by default.

diff --git a/Engine/Plugins/Editor/ScriptableToolsEditorMode/Source/ScriptableToolsEditorMode/Public/ScriptableToolsEditorModeSettings.h b/Engine/Plugins/Editor/ScriptableToolsEditorMode/Source/ScriptableToolsEditorMode/Public/ScriptableToolsEditorModeSettings.h
--- a/Engine/Plugins/Editor/ScriptableToolsEditorMode/Source/ScriptableToolsEditorMode/Public/ScriptableToolsEditorModeSettings.h
+++ b/Engine/Plugins/Editor/ScriptableToolsEditorMode/Source/ScriptableToolsEditorMode/Public/ScriptableToolsEditorModeSettings.h
@@ -35,7 +35,7 @@ public:
 
 	/** If true, Tool buttons will always be shown when in a Tool. By default they will be hidden. */
 	UPROPERTY(config, EditAnywhere, Category = "Scriptable Tools Mode|UI Customization")
-	bool bAlwaysShowToolButtons = false;
+	bool bAlwaysShowToolButtons = true;
 
 	bool RegisterAllTools() const {	return ToolRegistrationFilters.GetGroups().IsEmpty(); }
 
