ScriptableTools: Fixed an issue where ToolkitBuilder categories would duplicate when compiling tools with the mode open.

diff --git a/Engine/Plugins/Editor/ScriptableToolsEditorMode/Source/ScriptableToolsEditorMode/Private/ScriptableToolsEditorModeToolkit.cpp b/Engine/Plugins/Editor/ScriptableToolsEditorMode/Source/ScriptableToolsEditorMode/Private/ScriptableToolsEditorModeToolkit.cpp
--- a/Engine/Plugins/Editor/ScriptableToolsEditorMode/Source/ScriptableToolsEditorMode/Private/ScriptableToolsEditorModeToolkit.cpp
+++ b/Engine/Plugins/Editor/ScriptableToolsEditorMode/Source/ScriptableToolsEditorMode/Private/ScriptableToolsEditorModeToolkit.cpp
@@ -1177,6 +1177,8 @@ void FScriptableToolsEditorModeToolkit::RebuildModeToolPaletteWidgets()
 
 void FScriptableToolsEditorModeToolkit::RebuildModeToolkitBuilderPalettes()
 {
+	ToolkitBuilder->InitializeCategoryToolbar(true);
+	
 	const FScriptableToolsEditorModeManagerCommands& Commands = FScriptableToolsEditorModeManagerCommands::Get();
 	const TSharedPtr<FUICommandList> CommandList = GetToolkitCommands();
 	TMap<FName, TArray<TSharedPtr<FUICommandInfo>>> CommandLists = OwningEditorMode->GetModeCommands();
@@ -1221,7 +1223,11 @@ void FScriptableToolsEditorModeToolkit::RebuildModeToolkitBuilderPalettes()
 	{
 		ToolkitBuilder->SetActivePaletteOnLoad(ActivePaletteOnLoad.Get());
 	}
-	ToolkitBuilder->RefreshCategoryToolbarWidget(false);
+	
+	// Invoke UpdateWidget to update the CategoryToolbarVisibility on the ToolkitBuilder.
+	// This must be called prior to RefreshCategoryToolbarWidget since the visibility is not registered as a lambda.
+	ToolkitBuilder->UpdateWidget();
+	ToolkitBuilder->RefreshCategoryToolbarWidget(true);
 }
 
 #undef LOCTEXT_NAMESPACE
