ScriptableTools: Fixed issue with ToolkitBuilder palette building using incorrect ToolIdentifier name when locating matched tool commands.

diff --git a/Engine/Plugins/Editor/ScriptableToolsEditorMode/Source/ScriptableToolsEditorMode/Private/ScriptableToolsEditorMode.cpp b/Engine/Plugins/Editor/ScriptableToolsEditorMode/Source/ScriptableToolsEditorMode/Private/ScriptableToolsEditorMode.cpp
--- a/Engine/Plugins/Editor/ScriptableToolsEditorMode/Source/ScriptableToolsEditorMode/Private/ScriptableToolsEditorMode.cpp
+++ b/Engine/Plugins/Editor/ScriptableToolsEditorMode/Source/ScriptableToolsEditorMode/Private/ScriptableToolsEditorMode.cpp
@@ -236,7 +236,6 @@ void UScriptableToolsEditorMode::RebuildScriptableToolSet()
 			if (ToolCDO)
 			{
 				FString ToolNameString = ToolCDO->ToolName.IsEmpty() ? ToolIdentifier : ToolCDO->ToolName.ToString();
-				const FName ToolName(ToolNameString);
 				const FText ToolLabel = ToolCDO->ToolLongName.IsEmpty() ? FText::FromString(ToolNameString) : ToolCDO->ToolLongName;
 				const FText ToolTooltip = ToolCDO->ToolTooltip.IsEmpty() ? ToolLabel : ToolCDO->ToolTooltip;
 				TSharedPtr<FUICommandInfo> ToolCommand = ToolManagerCommands.RegisterCommand(
diff --git a/Engine/Plugins/Editor/ScriptableToolsEditorMode/Source/ScriptableToolsEditorMode/Private/ScriptableToolsEditorModeToolkit.cpp b/Engine/Plugins/Editor/ScriptableToolsEditorMode/Source/ScriptableToolsEditorMode/Private/ScriptableToolsEditorModeToolkit.cpp
--- a/Engine/Plugins/Editor/ScriptableToolsEditorMode/Source/ScriptableToolsEditorMode/Private/ScriptableToolsEditorModeToolkit.cpp
+++ b/Engine/Plugins/Editor/ScriptableToolsEditorMode/Source/ScriptableToolsEditorMode/Private/ScriptableToolsEditorModeToolkit.cpp
@@ -1201,19 +1201,21 @@ void FScriptableToolsEditorModeToolkit::RebuildModeToolkitBuilderPalettes()
 					continue;
 				}
 
-				bool bFound = false;
-				TSharedPtr<FUICommandInfo> ToolCommand = Commands.FindToolByName(ToolCDO->ToolName.ToString(), bFound);
-				if (bFound)
+				bool bFoundToolCommand = false;
+				FString ToolIdentifier;
+				Tool.ToolClass->GetClassPathName().ToString(ToolIdentifier);
+				TSharedPtr<FUICommandInfo> ToolCommand = Commands.FindToolByName(ToolIdentifier, bFoundToolCommand);
+				if (bFoundToolCommand)
 				{
 					PaletteCommandList.Add(ToolCommand);
 				}
 			}
 
-			bool bFound = false;
+			bool bFoundLoadPaletteCommand = false;
 			const FString LoadPaletteName = "LoadPalette" + PaletteName.ToString();
-			TSharedPtr<FUICommandInfo> LoadPaletteCommand = Commands.FindToolByName(LoadPaletteName, bFound);
+			TSharedPtr<FUICommandInfo> LoadPaletteCommand = Commands.FindToolByName(LoadPaletteName, bFoundLoadPaletteCommand);
 			ActivePaletteOnLoad = LoadPaletteCommand;
-			ensure(bFound);
+			ensure(bFoundLoadPaletteCommand);
 			ToolkitBuilder->AddPalette(
 				MakeShareable( new FToolPalette(LoadPaletteCommand.ToSharedRef(), PaletteCommandList)) );
 		}
