ScriptableTools: Fixed uncategorized tools not showing in the ToolkitBuilder path.

diff --git a/Engine/Plugins/Editor/ScriptableToolsEditorMode/Source/ScriptableToolsEditorMode/Private/ScriptableToolsEditorModeManagerCommands.cpp b/Engine/Plugins/Editor/ScriptableToolsEditorMode/Source/ScriptableToolsEditorMode/Private/ScriptableToolsEditorModeManagerCommands.cpp
--- a/Engine/Plugins/Editor/ScriptableToolsEditorMode/Source/ScriptableToolsEditorMode/Private/ScriptableToolsEditorModeManagerCommands.cpp
+++ b/Engine/Plugins/Editor/ScriptableToolsEditorMode/Source/ScriptableToolsEditorMode/Private/ScriptableToolsEditorModeManagerCommands.cpp
@@ -20,7 +20,7 @@ FScriptableToolsEditorModeManagerCommands::FScriptableToolsEditorModeManagerComm
 
 
 
-TSharedPtr<FUICommandInfo> FScriptableToolsEditorModeManagerCommands::FindToolByName(FString Name, bool& bFound) const
+TSharedPtr<FUICommandInfo> FScriptableToolsEditorModeManagerCommands::FindToolByName(const FString& Name, bool& bFound) const
 {
 	bFound = false;
 	for (const FStartToolCommand& Command : RegisteredTools)
diff --git a/Engine/Plugins/Editor/ScriptableToolsEditorMode/Source/ScriptableToolsEditorMode/Private/ScriptableToolsEditorModeToolkit.cpp b/Engine/Plugins/Editor/ScriptableToolsEditorMode/Source/ScriptableToolsEditorMode/Private/ScriptableToolsEditorModeToolkit.cpp
--- a/Engine/Plugins/Editor/ScriptableToolsEditorMode/Source/ScriptableToolsEditorMode/Private/ScriptableToolsEditorModeToolkit.cpp
+++ b/Engine/Plugins/Editor/ScriptableToolsEditorMode/Source/ScriptableToolsEditorMode/Private/ScriptableToolsEditorModeToolkit.cpp
@@ -35,6 +35,7 @@
 
 #define LOCTEXT_NAMESPACE "FScriptableToolsEditorModeToolkit"
 
+static const FName CustomToolsTabName(TEXT("Custom"));
 
 class SToolPaletteLoadBar : public SCompoundWidget
 {
@@ -610,18 +611,23 @@ void FScriptableToolsEditorModeToolkit::UpdateActiveToolCategories()
 	ScriptableTools->ForEachScriptableTool( [&](UClass* ToolClass, UBaseScriptableToolBuilder* ToolBuilder) 
 	{
 		UScriptableInteractiveTool* ToolCDO = Cast<UScriptableInteractiveTool>(ToolClass->GetDefaultObject());
-		if (ToolCDO->ToolCategory.IsEmpty() == false)
+		const bool bIsCategoryEmpty = ToolCDO->ToolCategory.IsEmpty();
+
+		// For legacy UI, don't process empty categories here.
+		if (!HasToolkitBuilder() && bIsCategoryEmpty)
 		{
-			FName CategoryName(ToolCDO->ToolCategory.ToString());
-			if (ActiveToolCategories.Contains(CategoryName) == false)
-			{
-				ActiveToolCategories.Add(CategoryName, {});
-			}
+			return;
+		}
 
-			if (TArray<FScriptableToolData>* CategoryTools = ActiveToolCategories.Find(CategoryName))
-			{
-				CategoryTools->Add(FScriptableToolData{ToolCDO->ToolCategory, ToolClass, ToolBuilder});
-			}
+		FName CategoryName = bIsCategoryEmpty ? CustomToolsTabName : FName(ToolCDO->ToolCategory.ToString()); 
+		if (ActiveToolCategories.Contains(CategoryName) == false)
+		{
+			ActiveToolCategories.Add(CategoryName, {});
+		}
+
+		if (TArray<FScriptableToolData>* CategoryTools = ActiveToolCategories.Find(CategoryName))
+		{
+			CategoryTools->Add(FScriptableToolData{ToolCDO->ToolCategory, ToolClass, ToolBuilder});
 		}
 	});
 }
@@ -693,15 +699,12 @@ FText FScriptableToolsEditorModeToolkit::GetBaseToolkitName() const
 	return NSLOCTEXT("ScriptableToolsEditorModeToolkit", "DisplayName", "ScriptableToolsEditorMode Tool");
 }
 
-static const FName CustomToolsTabName(TEXT("CustomTools"));
-
 void FScriptableToolsEditorModeToolkit::GetToolPaletteNames(TArray<FName>& PaletteNames) const
 {
 	PaletteNames.Reset();
 	bool bFoundUncategorized = false;
 
 	UScriptableToolsEditorMode* EditorMode = Cast<UScriptableToolsEditorMode>(GetScriptableEditorMode());
-	UInteractiveToolManager* ToolManager = EditorMode->GetToolManager();
 	UScriptableToolSet* ScriptableTools = EditorMode->GetActiveScriptableTools();
 	ScriptableTools->ForEachScriptableTool( [&](UClass* ToolClass, UBaseScriptableToolBuilder* ToolBuilder) 
 	{
diff --git a/Engine/Plugins/Editor/ScriptableToolsEditorMode/Source/ScriptableToolsEditorMode/Public/ScriptableToolsEditorModeManagerCommands.h b/Engine/Plugins/Editor/ScriptableToolsEditorMode/Source/ScriptableToolsEditorMode/Public/ScriptableToolsEditorModeManagerCommands.h
--- a/Engine/Plugins/Editor/ScriptableToolsEditorMode/Source/ScriptableToolsEditorMode/Public/ScriptableToolsEditorModeManagerCommands.h
+++ b/Engine/Plugins/Editor/ScriptableToolsEditorMode/Source/ScriptableToolsEditorMode/Public/ScriptableToolsEditorModeManagerCommands.h
@@ -25,7 +25,7 @@ public:
 	/**
 	 * Find Tool start-command below by registered name (tool icon name in Mode palette)
 	 */
-	TSharedPtr<FUICommandInfo> FindToolByName(FString Name, bool& bFound) const;
+	TSharedPtr<FUICommandInfo> FindToolByName(const FString& Name, bool& bFound) const;
 
 	TSharedPtr<FUICommandInfo> RegisterCommand(
 		FName InCommandName,
