ScriptableTools: Fixed a sporadic crash that could occur if a ScriptableTool invoked SpawnActorFromClass during tool Setup.

The OnBlueprintPreCompile callback will now:
1. Only deactivate the active tool if the active tool is of the same type as the BP generated class.
2. Only rebuild the Scriptable Toolset if the compiled blueprint is a ScriptableInteractiveTool.
3. Rebuilding the Scriptable Toolset on BP compiles now occurs on Tick.

diff --git a/Engine/Plugins/Editor/ScriptableToolsEditorMode/Source/ScriptableToolsEditorMode/Private/ScriptableToolsEditorMode.cpp b/Engine/Plugins/Editor/ScriptableToolsEditorMode/Source/ScriptableToolsEditorMode/Private/ScriptableToolsEditorMode.cpp
--- a/Engine/Plugins/Editor/ScriptableToolsEditorMode/Source/ScriptableToolsEditorMode/Private/ScriptableToolsEditorMode.cpp
+++ b/Engine/Plugins/Editor/ScriptableToolsEditorMode/Source/ScriptableToolsEditorMode/Private/ScriptableToolsEditorMode.cpp
@@ -34,6 +34,8 @@
 #include "ContextObjectStore.h"
 
 #include "Engine/StreamableManager.h"
+#include "Engine/Blueprint.h"
+#include "Kismet2/BlueprintEditorUtils.h"
 
 
 
@@ -136,6 +138,12 @@ void UScriptableToolsEditorMode::Tick(FEditorViewportClient* ViewportClient, flo
 		FScriptableToolsEditorModeToolkit* ModeToolkit = (FScriptableToolsEditorModeToolkit*)Toolkit.Get();
 		ModeToolkit->EnableShowRealtimeWarning(ViewportClient->IsRealtime() == false);
 	}
+
+	if (bRebuildScriptableToolSetOnTick)
+	{
+		RebuildScriptableToolSet();
+		bRebuildScriptableToolSetOnTick = false;
+	}
 }
 
 
@@ -184,7 +192,6 @@ void UScriptableToolsEditorMode::Enter()
 	GetToolManager()->SelectActiveToolType(EToolSide::Left, TEXT("BeginMeshInspectorTool"));
 
 	BlueprintPreCompileHandle = GEditor->OnBlueprintPreCompile().AddUObject(this, &UScriptableToolsEditorMode::OnBlueprintPreCompile); 
-	BlueprintCompiledHandle = GEditor->OnBlueprintCompiled().AddUObject(this, &UScriptableToolsEditorMode::OnBlueprintCompiled); 
 
 	// do any toolkit UI initialization that depends on the mode setup above
 	if (Toolkit.IsValid())
@@ -203,12 +210,25 @@ void UScriptableToolsEditorMode::RebuildScriptableToolSet()
 	auto UnregisterTools = [this]()
 	{
 		// unregister old tools from ToolManager
-		ScriptableTools->ForEachScriptableTool([&](UClass* ToolClass, UInteractiveToolBuilder* ToolBuilder)
+		ScriptableTools->ForEachScriptableTool([this](UClass* ToolClass, UInteractiveToolBuilder* ToolBuilder)
+		{
+			FString ToolIdentifier;
+			ToolClass->GetClassPathName().ToString(ToolIdentifier);
+			GetToolManager(EToolsContextScope::EdMode)->UnregisterToolType(ToolIdentifier);
+
+			if (Toolkit.IsValid() && FScriptableToolsEditorModeManagerCommands::IsRegistered())
 			{
-				FString ToolIdentifier;
-				ToolClass->GetClassPathName().ToString(ToolIdentifier);
-				GetToolManager(EToolsContextScope::EdMode)->UnregisterToolType(ToolIdentifier);
-			});
+				bool bFoundToolCommand = false;
+				const FScriptableToolsEditorModeManagerCommands& ToolManagerCommands = FScriptableToolsEditorModeManagerCommands::Get();
+				TSharedPtr<FUICommandInfo> ToolCommand = ToolManagerCommands.FindToolByName(ToolIdentifier, bFoundToolCommand);
+
+				if (ToolCommand)
+				{
+					const TSharedRef<FUICommandList>& CommandList = Toolkit->GetToolkitCommands();
+					CommandList->UnmapAction(ToolCommand);
+				}
+			}
+		});
 
 		FScriptableToolsEditorModeManagerCommands::Unregister();
 		
@@ -225,7 +245,7 @@ void UScriptableToolsEditorMode::RebuildScriptableToolSet()
 		
 		FScriptableToolsEditorModeManagerCommands& ToolManagerCommands = const_cast<FScriptableToolsEditorModeManagerCommands&>(FScriptableToolsEditorModeManagerCommands::Get());
 		// register each of them with ToolManager
-		ScriptableTools->ForEachScriptableTool([&](UClass* ToolClass, UInteractiveToolBuilder* ToolBuilder)
+		ScriptableTools->ForEachScriptableTool([this, &ToolManagerCommands](UClass* ToolClass, UInteractiveToolBuilder* ToolBuilder)
 		{
 			FString ToolIdentifier;
 			ToolClass->GetClassPathName().ToString(ToolIdentifier);
@@ -252,30 +272,21 @@ void UScriptableToolsEditorMode::RebuildScriptableToolSet()
 					CommandList->MapAction(ToolCommand,
 					FExecuteAction::CreateLambda([this, ToolClass, ToolIdentifier]()
 					{
-						if (GetToolManager()->SelectActiveToolType(EToolSide::Mouse, ToolIdentifier))
+						if (GetToolManager()->CanActivateTool(EToolSide::Mouse, ToolIdentifier))
 						{
-							if (GetToolManager()->CanActivateTool(EToolSide::Mouse, ToolIdentifier)) 
+							if (GetToolManager()->SelectActiveToolType(EToolSide::Mouse, ToolIdentifier)) 
 							{
 								GetToolManager()->ActivateTool(EToolSide::Mouse);
 							}
+							else
+							{
+								UE_LOG(LogTemp, Warning, TEXT("FAILED TO SET ACTIVE TOOL TYPE!"));
+							}
 						}
-						else
-						{
-							UE_LOG(LogTemp, Warning, TEXT("FAILED TO SET ACTIVE TOOL TYPE!"));
-						}
-
 					}),
 					FCanExecuteAction::CreateLambda([this, ToolClass, ToolIdentifier]()
 					{
-						if (GetToolManager()->SelectActiveToolType(EToolSide::Mouse, ToolIdentifier))
-						{
-							return GetToolManager()->CanActivateTool(EToolSide::Mouse, ToolIdentifier);
-						}
-						else
-						{
-							UE_LOG(LogTemp, Warning, TEXT("FAILED TO SET ACTIVE TOOL TYPE!"));
-							return false;
-						}
+						return GetToolManager()->CanActivateTool(EToolSide::Mouse, ToolIdentifier);
 					}),
 					FIsActionChecked::CreateLambda([this, ToolClass, ToolIdentifier]()
 					{
@@ -374,17 +385,31 @@ void UScriptableToolsEditorMode::InitializeModeContexts()
 
 void UScriptableToolsEditorMode::OnBlueprintPreCompile(UBlueprint* Blueprint)
 {
-	// would be nice to only do this if the Blueprint being compiled is being used for this tool...
-	if (GetToolManager()->HasActiveTool(EToolSide::Left))
+	if (!Blueprint)
 	{
-		GetToolManager()->DeactivateTool(EToolSide::Left, EToolShutdownType::Cancel);
+		return;
+	}
+	
+	if (Blueprint->GeneratedClass)
+	{
+		if (UInteractiveTool* ActiveTool = GetToolManager()->GetActiveTool(EToolSide::Left))
+		{
+			if (ActiveTool->IsA(Blueprint->GeneratedClass))
+			{
+				GetToolManager()->DeactivateTool(EToolSide::Left, EToolShutdownType::Cancel);
+			}
+		}
 	}
-}
-
 
-void UScriptableToolsEditorMode::OnBlueprintCompiled()
-{
-	RebuildScriptableToolSet();
+	// If this BP is a ScriptableInteractiveTool, schedule a rebuild of the toolset on tick.
+	// The OnTick is crucial in case the BP compile was initiated during Tool->Setup() for example
+	// which expects the tool to continue to exist after the Setup call. Invoking a rebuild
+	// inline would force deactivate all tools and result in a crash.
+	const UClass* NativeParentClass = FBlueprintEditorUtils::FindFirstNativeClass(Blueprint->ParentClass);
+	if (NativeParentClass->IsChildOf(UScriptableInteractiveTool::StaticClass()))
+	{
+		bRebuildScriptableToolSetOnTick = true;
+	}
 }
 
 void UScriptableToolsEditorMode::Exit()
@@ -393,7 +418,6 @@ void UScriptableToolsEditorMode::Exit()
 	ModeSettings->OnSettingChanged().RemoveAll(this);
 	
 	GEditor->OnBlueprintPreCompile().Remove(BlueprintPreCompileHandle);
-	GEditor->OnBlueprintCompiled().Remove(BlueprintCompiledHandle);
 
 	// exit any exclusive active tools w/ cancel
 	if (UInteractiveTool* ActiveTool = GetToolManager()->GetActiveTool(EToolSide::Left))
diff --git a/Engine/Plugins/Editor/ScriptableToolsEditorMode/Source/ScriptableToolsEditorMode/Public/ScriptableToolsEditorMode.h b/Engine/Plugins/Editor/ScriptableToolsEditorMode/Source/ScriptableToolsEditorMode/Public/ScriptableToolsEditorMode.h
--- a/Engine/Plugins/Editor/ScriptableToolsEditorMode/Source/ScriptableToolsEditorMode/Public/ScriptableToolsEditorMode.h
+++ b/Engine/Plugins/Editor/ScriptableToolsEditorMode/Source/ScriptableToolsEditorMode/Public/ScriptableToolsEditorMode.h
@@ -76,8 +76,6 @@ protected:
 
 	void OnToolsContextRender(IToolsContextRenderAPI* RenderAPI);
 
-	void FocusCameraAtCursorHotkey();
-
 	void AcceptActiveToolActionOrTool();
 	void CancelActiveToolActionOrTool();
 
@@ -86,7 +84,10 @@ protected:
 	void OnBlueprintPreCompile(UBlueprint* Blueprint);
 	FDelegateHandle BlueprintPreCompileHandle;
 
-	void OnBlueprintCompiled();
+	UE_DEPRECATED(5.6, "Deprecated in favor of OnBlueprintPreCompile")
+	void OnBlueprintCompiled() {}
+
+	UE_DEPRECATED(5.6, "Deprecated in favor of OnBlueprintPreCompile")
 	FDelegateHandle BlueprintCompiledHandle;
 
 	void InitializeModeContexts();
@@ -98,6 +99,8 @@ private:
 	TArray<TWeakObjectPtr<UScriptableToolContextObject>> ContextsToUpdateOnToolEnd;
 	TArray<TWeakObjectPtr<UScriptableToolContextObject>> ContextsToShutdown;
 
+	bool bRebuildScriptableToolSetOnTick = false;
+
 protected:
 	UPROPERTY()
 	TObjectPtr<UScriptableToolSet> ScriptableTools;
