ScriptableToolsMode: Fix issue where the toolkithost might be a nullptr when shutting down the mode.

diff --git a/Engine/Plugins/Editor/ScriptableToolsEditorMode/Source/ScriptableToolsEditorMode/Private/ScriptableToolsEditorMode.cpp b/Engine/Plugins/Editor/ScriptableToolsEditorMode/Source/ScriptableToolsEditorMode/Private/ScriptableToolsEditorMode.cpp
--- a/Engine/Plugins/Editor/ScriptableToolsEditorMode/Source/ScriptableToolsEditorMode/Private/ScriptableToolsEditorMode.cpp
+++ b/Engine/Plugins/Editor/ScriptableToolsEditorMode/Source/ScriptableToolsEditorMode/Private/ScriptableToolsEditorMode.cpp
@@ -371,10 +371,16 @@ void UScriptableToolsEditorMode::InitializeModeContexts()
 	ViewportWidgetAPI = NewObject<UScriptableToolViewportWidgetAPI>();
 	ViewportWidgetAPI->Initialize(
 		[this](TSharedRef<SWidget> InOverlaidWidget) {
-			Toolkit->GetToolkitHost()->AddViewportOverlayWidget(InOverlaidWidget);
+			if (Toolkit.IsValid() && Toolkit->IsHosted())
+			{
+				Toolkit->GetToolkitHost()->AddViewportOverlayWidget(InOverlaidWidget);
+			}
 		},
 		[this](TSharedRef<SWidget> InOverlaidWidget) {
-			Toolkit->GetToolkitHost()->RemoveViewportOverlayWidget(InOverlaidWidget);
+			if (Toolkit.IsValid() && Toolkit->IsHosted())
+			{
+				Toolkit->GetToolkitHost()->RemoveViewportOverlayWidget(InOverlaidWidget);
+			}
 		}
 		);
 	AddContextObject(ViewportWidgetAPI);
