ScriptableTools: Fixed issue with accept/cancel/complete actions not firing in ScriptableTools mode.

diff --git a/Engine/Plugins/Editor/ScriptableToolsEditorMode/Source/ScriptableToolsEditorMode/Private/ScriptableToolsEditorMode.cpp b/Engine/Plugins/Editor/ScriptableToolsEditorMode/Source/ScriptableToolsEditorMode/Private/ScriptableToolsEditorMode.cpp
--- a/Engine/Plugins/Editor/ScriptableToolsEditorMode/Source/ScriptableToolsEditorMode/Private/ScriptableToolsEditorMode.cpp
+++ b/Engine/Plugins/Editor/ScriptableToolsEditorMode/Source/ScriptableToolsEditorMode/Private/ScriptableToolsEditorMode.cpp
@@ -242,6 +242,9 @@ void UScriptableToolsEditorMode::RebuildScriptableToolSet()
 	auto RegisterTools = [this]()
 	{
 		FScriptableToolsEditorModeManagerCommands::Register();
+
+		// Rebind mode commands.
+		BindCommands();
 		
 		FScriptableToolsEditorModeManagerCommands& ToolManagerCommands = const_cast<FScriptableToolsEditorModeManagerCommands&>(FScriptableToolsEditorModeManagerCommands::Get());
 		// register each of them with ToolManager
