ScriptableTools: Fixed 'no world was found for object' warnings when invoking WorldContextObject functions.

diff --git a/Engine/Plugins/Editor/ScriptableToolsEditorMode/Source/ScriptableToolsEditorMode/Private/ScriptableToolsEditorMode.cpp b/Engine/Plugins/Editor/ScriptableToolsEditorMode/Source/ScriptableToolsEditorMode/Private/ScriptableToolsEditorMode.cpp
--- a/Engine/Plugins/Editor/ScriptableToolsEditorMode/Source/ScriptableToolsEditorMode/Private/ScriptableToolsEditorMode.cpp
+++ b/Engine/Plugins/Editor/ScriptableToolsEditorMode/Source/ScriptableToolsEditorMode/Private/ScriptableToolsEditorMode.cpp
@@ -496,6 +496,10 @@ void UScriptableToolsEditorMode::OnToolPostBuild(
 	UInteractiveToolManager* InToolManager, EToolSide InSide, 
 	UInteractiveTool* InBuiltTool, UInteractiveToolBuilder* InToolBuilder, const FToolBuilderState& ToolState)
 {
+	if (UScriptableInteractiveTool* ScriptTool = Cast<UScriptableInteractiveTool>(InBuiltTool))
+	{
+		ScriptTool->SetTargetWorld(GetInteractiveToolsContext()->GetWorld());
+	}
 }
 
 void UScriptableToolsEditorMode::OnToolStarted(UInteractiveToolManager* Manager, UInteractiveTool* Tool)
