ScriptableTools: Fixed a potential crash when re-registering Scriptable tools during mode shutdown.

diff --git a/Engine/Plugins/Editor/ScriptableToolsEditorMode/Source/ScriptableToolsEditorMode/Private/ScriptableToolsEditorMode.cpp b/Engine/Plugins/Editor/ScriptableToolsEditorMode/Source/ScriptableToolsEditorMode/Private/ScriptableToolsEditorMode.cpp
--- a/Engine/Plugins/Editor/ScriptableToolsEditorMode/Source/ScriptableToolsEditorMode/Private/ScriptableToolsEditorMode.cpp
+++ b/Engine/Plugins/Editor/ScriptableToolsEditorMode/Source/ScriptableToolsEditorMode/Private/ScriptableToolsEditorMode.cpp
@@ -527,6 +527,11 @@ void UScriptableToolsEditorMode::OnToolEnded(UInteractiveToolManager* Manager, U
 
 void UScriptableToolsEditorMode::BindCommands()
 {
+	if (!Toolkit.IsValid())
+	{
+		return;
+	}
+	
 	const FScriptableToolsEditorModeManagerCommands& ToolManagerCommands = FScriptableToolsEditorModeManagerCommands::Get();
 	const TSharedRef<FUICommandList>& CommandList = Toolkit->GetToolkitCommands();
 
