ScriptableTools: Added support for setting a custom UUserWidget above the tool details panel

diff --git a/Engine/Plugins/Editor/ScriptableToolsEditorMode/Source/ScriptableToolsEditorMode/Private/ScriptableToolsEditorModeToolkit.cpp b/Engine/Plugins/Editor/ScriptableToolsEditorMode/Source/ScriptableToolsEditorMode/Private/ScriptableToolsEditorModeToolkit.cpp
--- a/Engine/Plugins/Editor/ScriptableToolsEditorMode/Source/ScriptableToolsEditorMode/Private/ScriptableToolsEditorModeToolkit.cpp
+++ b/Engine/Plugins/Editor/ScriptableToolsEditorMode/Source/ScriptableToolsEditorMode/Private/ScriptableToolsEditorModeToolkit.cpp
@@ -464,6 +464,8 @@ void FScriptableToolsEditorModeToolkit::Init(const TSharedPtr<IToolkitHost>& Ini
 			ToolPaletteLoadBar->AsShared()
 		];
 
+	ToolDetailsHeader = SNew(SVerticalBox);
+
 	if (bUsesToolkitBuilder)
 	{
 		RegisterPalettes();
@@ -497,7 +499,8 @@ void FScriptableToolsEditorModeToolkit::Init(const TSharedPtr<IToolkitHost>& Ini
 		ToolkitSections->ModeWarningArea = ModeWarningArea;
 		ToolkitSections->DetailsView = ModeDetailsView;
 		ToolkitSections->ToolWarningArea = ToolWarningArea;
-		ToolkitSections->Footer = ToolPaletteHeader;
+		ToolkitSections->Header = ToolPaletteHeader;
+		ToolkitSections->ToolHeader = ToolDetailsHeader;
 
 		SAssignNew(ToolkitWidget, SBorder)
 		.HAlign(HAlign_Fill)
@@ -958,6 +961,17 @@ void FScriptableToolsEditorModeToolkit::OnToolStarted(UInteractiveToolManager* M
 		{
 			ActiveToolName = ScriptableTool->ToolName;
 		}
+
+		if (HasToolkitBuilder())
+		{
+			TSharedRef<SWidget> ActiveToolPanelHeaderWidget = ScriptableTool->GetToolPanelToolHeaderWidget(); 
+			ToolDetailsHeader->ClearChildren();
+			ToolDetailsHeader->AddSlot()
+			.AutoHeight()
+			[
+				ActiveToolPanelHeaderWidget
+			];
+		}
 	}
 
 	if (HasToolkitBuilder())
@@ -994,6 +1008,8 @@ void FScriptableToolsEditorModeToolkit::OnToolEnded(UInteractiveToolManager* Man
 	{
 		ToolkitBuilder->SetActiveToolDisplayName(FText::GetEmpty());
 		ToolkitBuilder->SetActivePaletteCommandsVisibility(EVisibility::Visible);
+
+		ToolDetailsHeader->ClearChildren();
 	}
 	
 	ModeHeaderArea->SetVisibility(EVisibility::Visible);
diff --git a/Engine/Plugins/Editor/ScriptableToolsEditorMode/Source/ScriptableToolsEditorMode/Public/ScriptableToolsEditorModeToolkit.h b/Engine/Plugins/Editor/ScriptableToolsEditorMode/Source/ScriptableToolsEditorMode/Public/ScriptableToolsEditorModeToolkit.h
--- a/Engine/Plugins/Editor/ScriptableToolsEditorMode/Source/ScriptableToolsEditorMode/Public/ScriptableToolsEditorModeToolkit.h
+++ b/Engine/Plugins/Editor/ScriptableToolsEditorMode/Source/ScriptableToolsEditorMode/Public/ScriptableToolsEditorModeToolkit.h
@@ -121,6 +121,7 @@ private:
 	TSharedPtr<SWidget> ToolPaletteHeader;
 	TSharedPtr<SWidget> ToolPaletteTagPanel;
 	TSharedPtr<SWidget> ToolPaletteLoadBar;
+	TSharedPtr<SVerticalBox> ToolDetailsHeader;
 
 	bool bShowRealtimeWarning = false;
 	UE_API void UpdateShowWarnings();
diff --git a/Engine/Plugins/Runtime/ScriptableToolsFramework/Source/ScriptableToolsFramework/Private/ScriptableInteractiveTool.cpp b/Engine/Plugins/Runtime/ScriptableToolsFramework/Source/ScriptableToolsFramework/Private/ScriptableInteractiveTool.cpp
--- a/Engine/Plugins/Runtime/ScriptableToolsFramework/Source/ScriptableToolsFramework/Private/ScriptableInteractiveTool.cpp
+++ b/Engine/Plugins/Runtime/ScriptableToolsFramework/Source/ScriptableToolsFramework/Private/ScriptableInteractiveTool.cpp
@@ -460,8 +460,8 @@ void UScriptableInteractiveTool::SetTargetWorld(UWorld* World)
 }
 
 UWorld* UScriptableInteractiveTool::GetWorld() const 
-{ 
-	return TargetWorld.Get(); 
+{
+	return TargetWorld.Get() ; 
 }
 
 
@@ -1324,6 +1324,28 @@ UScriptableToolTriangleSet* UScriptableInteractiveTool::AddTriangleSet()
 }
 
 
+void UScriptableInteractiveTool::SetToolPanelToolHeaderWidget(UUserWidget* Widget)
+{
+	ToolPanelToolHeaderWidget = Widget;
+}
+
+TSharedRef<SWidget> UScriptableInteractiveTool::GetToolPanelToolHeaderWidget() const
+{
+	if (ToolPanelToolHeaderWidget)
+	{
+		TSharedRef<SWidget> WrappedWidget = ToolPanelToolHeaderWidget->TakeDerivedWidget<SObjectWidget>([](UUserWidget* Widget, TSharedRef<SWidget> Content)
+		{
+			return SNew(SObjectWidget, Widget)
+				[
+					Content
+				];
+		});
+		return WrappedWidget;
+	}
+	return SNullWidget::NullWidget;
+}
+
+
 void UScriptableInteractiveTool::SetOverlayWidget(UUserWidget* Widget, bool bMakeDraggable )
 {
 	UContextObjectStore* ContextStore = GetToolManager()->GetContextObjectStore();
diff --git a/Engine/Plugins/Runtime/ScriptableToolsFramework/Source/ScriptableToolsFramework/Public/ScriptableInteractiveTool.h b/Engine/Plugins/Runtime/ScriptableToolsFramework/Source/ScriptableToolsFramework/Public/ScriptableInteractiveTool.h
--- a/Engine/Plugins/Runtime/ScriptableToolsFramework/Source/ScriptableToolsFramework/Public/ScriptableInteractiveTool.h
+++ b/Engine/Plugins/Runtime/ScriptableToolsFramework/Source/ScriptableToolsFramework/Public/ScriptableInteractiveTool.h
@@ -27,6 +27,7 @@ class UScriptableToolLineSet;
 class UScriptableToolPointSet;
 class UScriptableToolTriangleSet;
 class UUserWidget;
+class SWidget;
 
 UENUM(BlueprintType)
 enum class EToolsFrameworkOutcomePins : uint8
@@ -908,6 +909,30 @@ protected:
 
 public:
 
+	// Tool Panel Widget API
+
+	/**
+	 * Set a UUserWidget to display at the head of the Scriptable Tool panel details view.
+	 * Will only take effect when set during OnScriptSetup.
+	 * @param Widget The UUserWidget to set.
+	 */
+	UFUNCTION(BlueprintCallable, Category = "ScriptableTool|Widgets")
+	UE_API void SetToolPanelToolHeaderWidget(UUserWidget* Widget);
+	
+	/**
+	 * Retrieve an SWidget reference from the Scriptable Tool panel details view
+	 * header widget.
+	 * @return A reference to an SWidget
+	 */
+	UE_API TSharedRef<SWidget> GetToolPanelToolHeaderWidget() const;
+
+private:
+
+	UPROPERTY(Transient)
+	TObjectPtr<UUserWidget> ToolPanelToolHeaderWidget;
+
+public:
+	
 	// Widget Overlay API
 
 	UFUNCTION(BlueprintCallable, Category = "ScriptableTool|Widgets")
