ScriptableTools: Fixed crash when ScriptableToolsModeManagerCommands is not yet registered.

diff --git a/Engine/Plugins/Editor/ScriptableToolsEditorMode/Source/ScriptableToolsEditorMode/Private/ScriptableToolsEditorMode.cpp b/Engine/Plugins/Editor/ScriptableToolsEditorMode/Source/ScriptableToolsEditorMode/Private/ScriptableToolsEditorMode.cpp
--- a/Engine/Plugins/Editor/ScriptableToolsEditorMode/Source/ScriptableToolsEditorMode/Private/ScriptableToolsEditorMode.cpp
+++ b/Engine/Plugins/Editor/ScriptableToolsEditorMode/Source/ScriptableToolsEditorMode/Private/ScriptableToolsEditorMode.cpp
@@ -531,6 +531,17 @@ void UScriptableToolsEditorMode::BindCommands()
 	{
 		return;
 	}
+
+	// On subsequent runs of ScriptableToolsEditor mode, the FScriptableToolsEditorModeManagerCommands
+	// might be left in an unregistered state. Since this is only initially registered in the
+	// ScriptableToolsEditorModeModule, we need to potentially re-register the commands to handle this.
+	//
+	// With BindCommands being invoked during UEdMode::Enter(), this should ensure that the commands are
+	// in a valid state on mode enter.
+	if (!FScriptableToolsEditorModeManagerCommands::IsRegistered())
+	{
+		FScriptableToolsEditorModeManagerCommands::Register();
+	}
 	
 	const FScriptableToolsEditorModeManagerCommands& ToolManagerCommands = FScriptableToolsEditorModeManagerCommands::Get();
 	const TSharedRef<FUICommandList>& CommandList = Toolkit->GetToolkitCommands();
