ScriptableTools: Add support for texture-based tool icons

diff --git a/Engine/Plugins/Editor/ScriptableToolsEditorMode/Source/ScriptableToolsEditorMode/Private/ScriptableToolsEditorMode.cpp b/Engine/Plugins/Editor/ScriptableToolsEditorMode/Source/ScriptableToolsEditorMode/Private/ScriptableToolsEditorMode.cpp
--- a/Engine/Plugins/Editor/ScriptableToolsEditorMode/Source/ScriptableToolsEditorMode/Private/ScriptableToolsEditorMode.cpp
+++ b/Engine/Plugins/Editor/ScriptableToolsEditorMode/Source/ScriptableToolsEditorMode/Private/ScriptableToolsEditorMode.cpp
@@ -13,6 +13,7 @@
 #include "ScriptableToolsEditorModeToolkit.h"
 #include "ScriptableToolsEditorModeManagerCommands.h"
 #include "ScriptableToolsEditorModeSettings.h"
+#include "ScriptableToolsEditorModeStyle.h"
 
 #include "BaseGizmos/TransformGizmoUtil.h"
 #include "Snapping/ModelingSceneSnappingManager.h"
@@ -245,10 +246,35 @@ void UScriptableToolsEditorMode::RebuildScriptableToolSet()
 
 		// Rebind mode commands.
 		BindCommands();
-		
+
+		auto GetToolIconName = [](const FString& ToolIdentifier)
+		{
+			return FName(ToolIdentifier + ".Icon");
+		};
+
+		// Register tool icons
+		FScriptableToolsEditorModeStyle::FIconTextureMap IconTextureMap;
+		ScriptableTools->ForEachScriptableTool([this, GetToolIconName, &IconTextureMap](UClass* ToolClass, UInteractiveToolBuilder* ToolBuilder)
+		{
+			FString ToolIdentifier;
+			ToolClass->GetClassPathName().ToString(ToolIdentifier);
+
+			if (UScriptableInteractiveTool* ToolCDO = Cast<UScriptableInteractiveTool>(ToolClass->GetDefaultObject()))
+			{
+				if (ToolCDO->ToolIconTexture)
+				{
+					IconTextureMap.Emplace(GetToolIconName(ToolIdentifier), ToolCDO->ToolIconTexture);
+				}
+			}
+		});
+		if (IconTextureMap.Num() > 0)
+		{
+			FScriptableToolsEditorModeStyle::RegisterIconTextures(IconTextureMap);
+		}
+
+		// Register tools with ToolManager
 		FScriptableToolsEditorModeManagerCommands& ToolManagerCommands = const_cast<FScriptableToolsEditorModeManagerCommands&>(FScriptableToolsEditorModeManagerCommands::Get());
-		// register each of them with ToolManager
-		ScriptableTools->ForEachScriptableTool([this, &ToolManagerCommands](UClass* ToolClass, UInteractiveToolBuilder* ToolBuilder)
+		ScriptableTools->ForEachScriptableTool([this, &ToolManagerCommands, GetToolIconName](UClass* ToolClass, UInteractiveToolBuilder* ToolBuilder)
 		{
 			FString ToolIdentifier;
 			ToolClass->GetClassPathName().ToString(ToolIdentifier);
@@ -261,11 +287,17 @@ void UScriptableToolsEditorMode::RebuildScriptableToolSet()
 				FString ToolNameString = ToolCDO->ToolName.IsEmpty() ? ToolIdentifier : ToolCDO->ToolName.ToString();
 				const FText ToolLabel = ToolCDO->ToolLongName.IsEmpty() ? FText::FromString(ToolNameString) : ToolCDO->ToolLongName;
 				const FText ToolTooltip = ToolCDO->ToolTooltip.IsEmpty() ? ToolLabel : ToolCDO->ToolTooltip;
+				FSlateIcon ToolIcon;
+				if (ToolCDO->ToolIconTexture)
+				{
+					FName ToolIconName = GetToolIconName(ToolIdentifier);
+					ToolIcon = FSlateIcon(FScriptableToolsEditorModeStyle::Get()->GetStyleSetName(), ToolIconName);
+				}
 				TSharedPtr<FUICommandInfo> ToolCommand = ToolManagerCommands.RegisterCommand(
 					FName(ToolIdentifier),
 					ToolLabel,
 					ToolTooltip,
-					FSlateIcon(),
+					ToolIcon,
 					EUserInterfaceActionType::ToggleButton,
 					FInputChord());
 
diff --git a/Engine/Plugins/Editor/ScriptableToolsEditorMode/Source/ScriptableToolsEditorMode/Private/ScriptableToolsEditorModeStyle.cpp b/Engine/Plugins/Editor/ScriptableToolsEditorMode/Source/ScriptableToolsEditorMode/Private/ScriptableToolsEditorModeStyle.cpp
--- a/Engine/Plugins/Editor/ScriptableToolsEditorMode/Source/ScriptableToolsEditorMode/Private/ScriptableToolsEditorModeStyle.cpp
+++ b/Engine/Plugins/Editor/ScriptableToolsEditorMode/Source/ScriptableToolsEditorMode/Private/ScriptableToolsEditorModeStyle.cpp
@@ -15,6 +15,8 @@
 #define IMAGE_PLUGIN_BRUSH_SVG( RelativePath, ... ) FSlateVectorImageBrush(FScriptableToolsEditorModeStyle::InContent( RelativePath, ".svg"), __VA_ARGS__)
 #define RootToContentDir StyleSet->RootToContentDir
 
+TMap<FName, FSlateBrush*> FScriptableToolsEditorModeStyle::IconTextureBrushes;
+
 FString FScriptableToolsEditorModeStyle::InContent(const FString& RelativePath, const ANSICHAR* Extension)
 {
 	static FString ContentDir = IPluginManager::Get().FindPlugin(TEXT("ScriptableToolsEditorMode"))->GetContentDir();
@@ -120,6 +122,42 @@ bool FScriptableToolsEditorModeStyle::TryRegisterCustomIcon(FName StyleIdentifie
 	return true;
 }
 
+bool FScriptableToolsEditorModeStyle::RegisterIconTexture(FName StyleIdentifier, TObjectPtr<UTexture2D> InTexture)
+{
+	if (!InTexture)
+	{
+		return false;
+	}
+
+	RegisterIconTextures({{StyleIdentifier, InTexture}});
+	return true;
+}
+
+void FScriptableToolsEditorModeStyle::RegisterIconTextures(const FIconTextureMap& InIcons)
+{
+	// need to re-register style to be able to modify it
+	FSlateStyleRegistry::UnRegisterSlateStyle(*StyleSet.Get());
+
+	const FVector2D Icon20x20(20.0f, 20.0f);
+	for (const TPair<FName, TObjectPtr<UTexture2D>>& IconData : InIcons)
+	{
+		if (IconData.Value)
+		{
+			if (FSlateBrush** OldBrush = IconTextureBrushes.Find(IconData.Key))
+			{
+				delete *OldBrush;
+			}
+			FSlateBrush* NewBrush = new FSlateImageBrush(static_cast<UObject*>(IconData.Value), Icon20x20);
+			IconTextureBrushes.Add(IconData.Key, NewBrush);
+			StyleSet->Set(IconData.Key, NewBrush);
+		}
+	}
+	
+	FSlateStyleRegistry::RegisterSlateStyle(*StyleSet.Get());
+}
+
+
+
 END_SLATE_FUNCTION_BUILD_OPTIMIZATION
 
 #undef IMAGE_PLUGIN_BRUSH
@@ -134,4 +172,10 @@ void FScriptableToolsEditorModeStyle::Shutdown()
 		ensure(StyleSet.IsUnique());
 		StyleSet.Reset();
 	}
+
+	for (TPair<FName, FSlateBrush*> IconBrush : IconTextureBrushes)
+	{
+		delete IconBrush.Value;
+	}
+	IconTextureBrushes.Reset();
 }
diff --git a/Engine/Plugins/Editor/ScriptableToolsEditorMode/Source/ScriptableToolsEditorMode/Private/ScriptableToolsEditorModeToolkit.cpp b/Engine/Plugins/Editor/ScriptableToolsEditorMode/Source/ScriptableToolsEditorMode/Private/ScriptableToolsEditorModeToolkit.cpp
--- a/Engine/Plugins/Editor/ScriptableToolsEditorMode/Source/ScriptableToolsEditorMode/Private/ScriptableToolsEditorModeToolkit.cpp
+++ b/Engine/Plugins/Editor/ScriptableToolsEditorMode/Source/ScriptableToolsEditorMode/Private/ScriptableToolsEditorModeToolkit.cpp
@@ -831,7 +831,10 @@ void FScriptableToolsEditorModeToolkit::BuildToolPalette(FName PaletteIndex, cla
 
 		// if a custom icon is defined, try to find it, this can fail in many ways, in that case
 		// the default icon is kept
-		if (ToolCDO->CustomIconPath.IsEmpty() == false)
+PRAGMA_DISABLE_DEPRECATION_WARNINGS
+		const FString CustomIconPath = ToolCDO->CustomIconPath;
+PRAGMA_ENABLE_DEPRECATION_WARNINGS
+		if (CustomIconPath.IsEmpty() == false)
 		{
 			FName ToolIconToken( FString("ScriptableToolsEditorModeToolCommands.") + ToolIdentifier );
 
@@ -844,10 +847,10 @@ void FScriptableToolsEditorModeToolkit::BuildToolPalette(FName PaletteIndex, cla
 			FString PathPart, FilenamePart, ExtensionPart;
 			FPaths::Split(FullPathName, PathPart, FilenamePart, ExtensionPart);
 
-			FString FullIconPath = ToolCDO->CustomIconPath;
+			FString FullIconPath = CustomIconPath;
 			if (PathPart.StartsWith("/Game"))
 			{
-				FullIconPath = FPaths::ProjectContentDir() / ToolCDO->CustomIconPath;
+				FullIconPath = FPaths::ProjectContentDir() / CustomIconPath;
 			}
 			else
 			{
@@ -856,15 +859,17 @@ void FScriptableToolsEditorModeToolkit::BuildToolPalette(FName PaletteIndex, cla
 				if (PathParts.Num() > 0)
 				{
 					FString PluginContentDir = IPluginManager::Get().FindPlugin(PathParts[0])->GetContentDir();
-					FullIconPath = PluginContentDir / ToolCDO->CustomIconPath;
+					FullIconPath = PluginContentDir / CustomIconPath;
 				}
 				else  // something is wrong, fall back to project content dir
 				{
-					FullIconPath = FPaths::ProjectContentDir() / ToolCDO->CustomIconPath;
+					FullIconPath = FPaths::ProjectContentDir() / CustomIconPath;
 				}
 			}
 
-			if (FScriptableToolsEditorModeStyle::TryRegisterCustomIcon(ToolIconToken, FullIconPath, ToolCDO->CustomIconPath))
+PRAGMA_DISABLE_DEPRECATION_WARNINGS
+			if (FScriptableToolsEditorModeStyle::TryRegisterCustomIcon(ToolIconToken, FullIconPath, CustomIconPath))
+PRAGMA_ENABLE_DEPRECATION_WARNINGS
 			{
 				Icon = FSlateIcon(FScriptableToolsEditorModeStyle::Get()->GetStyleSetName(), ToolIconToken);
 			}
diff --git a/Engine/Plugins/Editor/ScriptableToolsEditorMode/Source/ScriptableToolsEditorMode/Public/ScriptableToolsEditorModeStyle.h b/Engine/Plugins/Editor/ScriptableToolsEditorMode/Source/ScriptableToolsEditorMode/Public/ScriptableToolsEditorModeStyle.h
--- a/Engine/Plugins/Editor/ScriptableToolsEditorMode/Source/ScriptableToolsEditorMode/Public/ScriptableToolsEditorModeStyle.h
+++ b/Engine/Plugins/Editor/ScriptableToolsEditorMode/Source/ScriptableToolsEditorMode/Public/ScriptableToolsEditorModeStyle.h
@@ -33,14 +33,22 @@ public:
 	// the user enters as (eg) a property/etc. To help with debugging, this path can be provided as ExternalRelativePath,
 	// and it will be printed out instead of the absolute path as part of the log error message.
 	//
+	UE_DEPRECATED(5.7, "Use RegisterIconTexture(s) instead.")
 	static UE_API bool TryRegisterCustomIcon(FName StyleIconIdentifier, FString FileNameWithPath, FString ExternalRelativePath);
 
+	// Register tool icons from a UTexture2D
+	using FIconTextureMap = TMap<FName, TObjectPtr<UTexture2D>>;
+	static UE_API bool RegisterIconTexture(FName StyleIdentifier, TObjectPtr<UTexture2D> InTexture);
+	static UE_API void RegisterIconTextures(const FIconTextureMap& InIcons);
+
 
 private:
 	static UE_API FString InContent(const FString& RelativePath, const ANSICHAR* Extension);
 
 private:
 	static UE_API TSharedPtr< class FSlateStyleSet > StyleSet;
+
+	static TMap<FName, FSlateBrush*> IconTextureBrushes;
 };
 
 #undef UE_API
diff --git a/Engine/Plugins/Runtime/ScriptableToolsFramework/Source/ScriptableToolsFramework/Public/ScriptableInteractiveTool.h b/Engine/Plugins/Runtime/ScriptableToolsFramework/Source/ScriptableToolsFramework/Public/ScriptableInteractiveTool.h
--- a/Engine/Plugins/Runtime/ScriptableToolsFramework/Source/ScriptableToolsFramework/Public/ScriptableInteractiveTool.h
+++ b/Engine/Plugins/Runtime/ScriptableToolsFramework/Source/ScriptableToolsFramework/Public/ScriptableInteractiveTool.h
@@ -22,6 +22,7 @@ class FCanvas;
 class FSceneView;
 
 class UPreviewGeometry;
+class UTexture2D;
 
 class UScriptableToolLineSet;
 class UScriptableToolPointSet;
@@ -342,10 +343,14 @@ public:
 	 * So for example if the Tool BP is in a plugin named MyCustomTools, the icon must be in 
 	 *      MyProject/Plugins/MyCustomTools/Content/<SubFolders>/MyToolIcon.png,
 	 * and the relative path to use here would be <SubFolders>/MyToolIcon.png
-	 */ 
-	UPROPERTY(BlueprintReadWrite, EditAnywhere, Category = "Scriptable Tool Settings", meta = (DisplayName = "Custom Icon Path"))
+	 */
+	UE_DEPRECATED(5.7, "Use ToolIconTexture instead.")
+	UPROPERTY(meta = (DisplayName = "Custom Icon Path", DeprecatedProperty, DeprecationMessage="Use ToolIconTexture instead."))
 	FString CustomIconPath;
 
+	UPROPERTY(EditDefaultsOnly, Category = "Scriptable Tool Settings", meta = (DisplayName = "Tool Icon"))
+	TObjectPtr<UTexture2D> ToolIconTexture;
+
 	/** A generic flag to indicate whether this Tool should be shown in the UE Editor. This may be interpreted in different ways */
 	UPROPERTY(BlueprintReadWrite, EditAnywhere, Category = "Scriptable Tool Settings", meta=(DisplayName="Visible in Editor"))
 	bool bShowToolInEditor = true;
