ScriptableTools: Disable false positive static analysis warning.

diff --git a/Engine/Plugins/Editor/ScriptableToolsEditorMode/Source/ScriptableToolsEditorMode/Private/ScriptableToolsEditorMode.cpp b/Engine/Plugins/Editor/ScriptableToolsEditorMode/Source/ScriptableToolsEditorMode/Private/ScriptableToolsEditorMode.cpp
--- a/Engine/Plugins/Editor/ScriptableToolsEditorMode/Source/ScriptableToolsEditorMode/Private/ScriptableToolsEditorMode.cpp
+++ b/Engine/Plugins/Editor/ScriptableToolsEditorMode/Source/ScriptableToolsEditorMode/Private/ScriptableToolsEditorMode.cpp
@@ -267,7 +267,7 @@ void UScriptableToolsEditorMode::RebuildScriptableToolSet()
 				}
 			}
 		});
-		if (IconTextureMap.Num() > 0)
+		if (IconTextureMap.Num() > 0) //-V547
 		{
 			FScriptableToolsEditorModeStyle::RegisterIconTextures(IconTextureMap);
 		}
