ScriptableTools: Fixed non-deterministic tool order in category palettes.

diff --git a/Engine/Plugins/Editor/ScriptableToolsEditorMode/Source/ScriptableToolsEditorMode/Private/ScriptableToolsEditorModeToolkit.cpp b/Engine/Plugins/Editor/ScriptableToolsEditorMode/Source/ScriptableToolsEditorMode/Private/ScriptableToolsEditorModeToolkit.cpp
--- a/Engine/Plugins/Editor/ScriptableToolsEditorMode/Source/ScriptableToolsEditorMode/Private/ScriptableToolsEditorModeToolkit.cpp
+++ b/Engine/Plugins/Editor/ScriptableToolsEditorMode/Source/ScriptableToolsEditorMode/Private/ScriptableToolsEditorModeToolkit.cpp
@@ -1235,6 +1235,11 @@ void FScriptableToolsEditorModeToolkit::RebuildModeToolkitBuilderPalettes()
 				}
 			}
 
+			Algo::SortBy(PaletteCommandList, [](const TSharedPtr<FUICommandInfo>& Data)
+			{
+				return Data->GetCommandName().ToString();
+			});
+
 			bool bFoundLoadPaletteCommand = false;
 			const FString LoadPaletteName = "LoadPalette" + PaletteName.ToString();
 			TSharedPtr<FUICommandInfo> LoadPaletteCommand = Commands.FindToolByName(LoadPaletteName, bFoundLoadPaletteCommand);
