Made realtime warning for modeling mode and scriptable tools mode only update when the app has focus, to avoid spurious warnings when the 'background task' override disables real time mode

diff --git a/Engine/Plugins/Editor/ModelingToolsEditorMode/Source/ModelingToolsEditorMode/Private/ModelingToolsEditorMode.cpp b/Engine/Plugins/Editor/ModelingToolsEditorMode/Source/ModelingToolsEditorMode/Private/ModelingToolsEditorMode.cpp
--- a/Engine/Plugins/Editor/ModelingToolsEditorMode/Source/ModelingToolsEditorMode/Private/ModelingToolsEditorMode.cpp
+++ b/Engine/Plugins/Editor/ModelingToolsEditorMode/Source/ModelingToolsEditorMode/Private/ModelingToolsEditorMode.cpp
@@ -272,7 +272,7 @@ void UModelingToolsEditorMode::Tick(FEditorViewportClient* ViewportClient, float
 {
 	Super::Tick(ViewportClient, DeltaTime);
 
-	if (Toolkit.IsValid())
+	if (Toolkit.IsValid() && FApp::HasFocus())
 	{
 		FModelingToolsEditorModeToolkit* ModelingToolkit = (FModelingToolsEditorModeToolkit*)Toolkit.Get();
 		ModelingToolkit->ShowRealtimeAndModeWarnings(ViewportClient->IsRealtime() == false);
diff --git a/Engine/Plugins/Editor/ScriptableToolsEditorMode/Source/ScriptableToolsEditorMode/Private/ScriptableToolsEditorMode.cpp b/Engine/Plugins/Editor/ScriptableToolsEditorMode/Source/ScriptableToolsEditorMode/Private/ScriptableToolsEditorMode.cpp
--- a/Engine/Plugins/Editor/ScriptableToolsEditorMode/Source/ScriptableToolsEditorMode/Private/ScriptableToolsEditorMode.cpp
+++ b/Engine/Plugins/Editor/ScriptableToolsEditorMode/Source/ScriptableToolsEditorMode/Private/ScriptableToolsEditorMode.cpp
@@ -134,7 +134,7 @@ void UScriptableToolsEditorMode::Tick(FEditorViewportClient* ViewportClient, flo
 {
 	Super::Tick(ViewportClient, DeltaTime);
 
-	if (Toolkit.IsValid())
+	if (Toolkit.IsValid() && FApp::HasFocus())
 	{
 		FScriptableToolsEditorModeToolkit* ModeToolkit = (FScriptableToolsEditorModeToolkit*)Toolkit.Get();
 		ModeToolkit->EnableShowRealtimeWarning(ViewportClient->IsRealtime() == false);
