Removed some unused code in ScriptableToolsEditorMode.

#rnx

diff --git a/Engine/Plugins/Editor/ScriptableToolsEditorMode/Source/ScriptableToolsEditorMode/Private/ScriptableToolsEditorMode.cpp b/Engine/Plugins/Editor/ScriptableToolsEditorMode/Source/ScriptableToolsEditorMode/Private/ScriptableToolsEditorMode.cpp
--- a/Engine/Plugins/Editor/ScriptableToolsEditorMode/Source/ScriptableToolsEditorMode/Private/ScriptableToolsEditorMode.cpp
+++ b/Engine/Plugins/Editor/ScriptableToolsEditorMode/Source/ScriptableToolsEditorMode/Private/ScriptableToolsEditorMode.cpp
@@ -44,16 +44,6 @@
 
 const FEditorModeID UScriptableToolsEditorMode::EM_ScriptableToolsEditorModeId = TEXT("EM_ScriptableToolsEditorMode");
 
-
-namespace
-{
-FString GetToolName(const UInteractiveTool& Tool)
-{
-	const FString* ToolName = FTextInspector::GetSourceString(Tool.GetToolInfo().ToolDisplayName);
-	return ToolName ? *ToolName : FString(TEXT("<Invalid ToolName>"));
-}
-}
-
 UScriptableToolsEditorMode::UScriptableToolsEditorMode()
 {
 	Info = FEditorModeInfo(
