Added a SetFocusInViewport function that's exposed to Blueprint. This will allow scriptable tools to give focus to the viewport so that keybindings pass thru without user interaction.

[RN] Added a SetFocusInViewport function that's exposed to Blueprint. This will allow scriptable tools to give focus to the viewport so that keybindings pass thru without user interaction.

diff --git a/Engine/Plugins/Editor/ScriptableToolsEditorMode/Source/ScriptableToolsEditorMode/Private/ScriptableToolsEditorMode.cpp b/Engine/Plugins/Editor/ScriptableToolsEditorMode/Source/ScriptableToolsEditorMode/Private/ScriptableToolsEditorMode.cpp
--- a/Engine/Plugins/Editor/ScriptableToolsEditorMode/Source/ScriptableToolsEditorMode/Private/ScriptableToolsEditorMode.cpp
+++ b/Engine/Plugins/Editor/ScriptableToolsEditorMode/Source/ScriptableToolsEditorMode/Private/ScriptableToolsEditorMode.cpp
@@ -37,6 +37,7 @@
 #include "Engine/StreamableManager.h"
 #include "Engine/Blueprint.h"
 #include "Kismet2/BlueprintEditorUtils.h"
+#include "Slate/SceneViewport.h"
 
 #include UE_INLINE_GENERATED_CPP_BY_NAME(ScriptableToolsEditorMode)
 
@@ -405,7 +406,6 @@ void UScriptableToolsEditorMode::InitializeModeContexts()
 	};
 
 	UScriptableToolViewportWidgetAPI* ViewportWidgetAPI = NewObject<UScriptableToolViewportWidgetAPI>();
-	ViewportWidgetAPI = NewObject<UScriptableToolViewportWidgetAPI>();
 	ViewportWidgetAPI->Initialize(
 		[this](TSharedRef<SWidget> InOverlaidWidget) {
 			if (Toolkit.IsValid() && Toolkit->IsHosted())
@@ -422,9 +422,42 @@ void UScriptableToolsEditorMode::InitializeModeContexts()
 		);
 	AddContextObject(ViewportWidgetAPI);
 
+	UScriptableToolViewportFocusAPI* ViewportFocusAPI = NewObject<UScriptableToolViewportFocusAPI>();
+	ViewportFocusAPI->Initialize(
+		[this]() {
+			SetFocusInViewport();
+		}
+		);
+	AddContextObject(ViewportFocusAPI);
 }
 
+void UScriptableToolsEditorMode::SetFocusInViewport()
+{
+	FEditorModeTools& ModeManager = GLevelEditorModeTools();
 
+	const FEditorViewportClient* ViewportClient = ModeManager.GetHoveredViewportClient();
+	if (!ViewportClient)
+	{
+		ViewportClient = ModeManager.GetFocusedViewportClient();
+	}
+	if (!ViewportClient)
+	{
+		return;
+	}
+
+	const TWeakPtr<SViewport> ViewportWidget = ViewportClient->GetEditorViewportWidget()->GetSceneViewport()->GetViewportWidget();
+	if (!ViewportWidget.IsValid())
+	{
+		return;
+	}
+
+	// set focus to viewport so that hotkeys are immediately detected
+	TSharedPtr<SWidget> ViewportContents = ViewportWidget.Pin()->GetContent();
+	FSlateApplication::Get().ForEachUser([&ViewportContents](FSlateUser& User)
+	{
+		User.SetFocus(ViewportContents.ToSharedRef());
+	});
+}
 
 void UScriptableToolsEditorMode::OnBlueprintPreCompile(UBlueprint* Blueprint)
 {
diff --git a/Engine/Plugins/Editor/ScriptableToolsEditorMode/Source/ScriptableToolsEditorMode/Public/ScriptableToolsEditorMode.h b/Engine/Plugins/Editor/ScriptableToolsEditorMode/Source/ScriptableToolsEditorMode/Public/ScriptableToolsEditorMode.h
--- a/Engine/Plugins/Editor/ScriptableToolsEditorMode/Source/ScriptableToolsEditorMode/Public/ScriptableToolsEditorMode.h
+++ b/Engine/Plugins/Editor/ScriptableToolsEditorMode/Source/ScriptableToolsEditorMode/Public/ScriptableToolsEditorMode.h
@@ -96,6 +96,8 @@ protected:
 
 	UE_API void RebuildScriptableToolSet();
 
+	void SetFocusInViewport();
+
 private:
 
 	TArray<TWeakObjectPtr<UScriptableToolContextObject>> ContextsToUpdateOnToolEnd;
diff --git a/Engine/Plugins/Runtime/ScriptableToolsFramework/Source/ScriptableToolsFramework/Private/ScriptableInteractiveTool.cpp b/Engine/Plugins/Runtime/ScriptableToolsFramework/Source/ScriptableToolsFramework/Private/ScriptableInteractiveTool.cpp
--- a/Engine/Plugins/Runtime/ScriptableToolsFramework/Source/ScriptableToolsFramework/Private/ScriptableInteractiveTool.cpp
+++ b/Engine/Plugins/Runtime/ScriptableToolsFramework/Source/ScriptableToolsFramework/Private/ScriptableInteractiveTool.cpp
@@ -1386,6 +1386,17 @@ void UScriptableInteractiveTool::ClearOverlayWidget()
 	}
 }
 
+void UScriptableInteractiveTool::SetFocusInViewport()
+{
+	UContextObjectStore* ContextStore = GetToolManager()->GetContextObjectStore();
+
+	UScriptableToolViewportFocusAPI* ViewportFocusAPI = ContextStore->FindContext< UScriptableToolViewportFocusAPI>();
+	if (ViewportFocusAPI)
+	{
+		ViewportFocusAPI->SetFocusInViewport();
+	}
+}
+
 FInputRayHit UScriptableToolsUtilityLibrary::MakeInputRayHit_Miss()
 {
 	return FInputRayHit();
diff --git a/Engine/Plugins/Runtime/ScriptableToolsFramework/Source/ScriptableToolsFramework/Public/ScriptableInteractiveTool.h b/Engine/Plugins/Runtime/ScriptableToolsFramework/Source/ScriptableToolsFramework/Public/ScriptableInteractiveTool.h
--- a/Engine/Plugins/Runtime/ScriptableToolsFramework/Source/ScriptableToolsFramework/Public/ScriptableInteractiveTool.h
+++ b/Engine/Plugins/Runtime/ScriptableToolsFramework/Source/ScriptableToolsFramework/Public/ScriptableInteractiveTool.h
@@ -951,6 +951,13 @@ public:
 
 	UPROPERTY(EditAnywhere, Transient, Category = "Scriptable Tool Settings")
 	FScriptableToolGroupSet GroupTags;
+
+public:
+
+	// Viewport Focus API
+
+	UFUNCTION(BlueprintCallable, Category = "ScriptableTools|Viewport")
+	UE_API void SetFocusInViewport();
 };
 
 
diff --git a/Engine/Plugins/Runtime/ScriptableToolsFramework/Source/ScriptableToolsFramework/Public/Utility/ScriptableToolContextObjects.h b/Engine/Plugins/Runtime/ScriptableToolsFramework/Source/ScriptableToolsFramework/Public/Utility/ScriptableToolContextObjects.h
--- a/Engine/Plugins/Runtime/ScriptableToolsFramework/Source/ScriptableToolsFramework/Public/Utility/ScriptableToolContextObjects.h
+++ b/Engine/Plugins/Runtime/ScriptableToolsFramework/Source/ScriptableToolsFramework/Public/Utility/ScriptableToolContextObjects.h
@@ -83,3 +83,33 @@ private:
 
 	TWeakPtr<SWidget> CurrentOverlaidWidget;
 };
+
+UCLASS(MinimalAPI)
+class UScriptableToolViewportFocusAPI : public UScriptableToolContextObject
+{
+	GENERATED_BODY()
+
+public:
+
+	void Initialize(TUniqueFunction<void()> SetFocusInViewportFuncIn)
+	{
+		SetFocusInViewportFunc = MoveTemp(SetFocusInViewportFuncIn);
+	}
+
+	void SetFocusInViewport()
+	{
+		SetFocusInViewportFunc();
+	}
+
+	virtual void Shutdown() override
+	{
+	}
+
+	virtual void OnToolEnded(UInteractiveTool* DeadTool) override
+	{
+	}
+
+private:
+
+	TUniqueFunction<void()> SetFocusInViewportFunc;
+};
