Added validation checks around pointers.

#rnx

diff --git a/Engine/Plugins/Editor/ScriptableToolsEditorMode/Source/ScriptableToolsEditorMode/Private/ScriptableToolsEditorMode.cpp b/Engine/Plugins/Editor/ScriptableToolsEditorMode/Source/ScriptableToolsEditorMode/Private/ScriptableToolsEditorMode.cpp
--- a/Engine/Plugins/Editor/ScriptableToolsEditorMode/Source/ScriptableToolsEditorMode/Private/ScriptableToolsEditorMode.cpp
+++ b/Engine/Plugins/Editor/ScriptableToolsEditorMode/Source/ScriptableToolsEditorMode/Private/ScriptableToolsEditorMode.cpp
@@ -445,18 +445,29 @@ void UScriptableToolsEditorMode::SetFocusInViewport()
 		return;
 	}
 
-	const TWeakPtr<SViewport> ViewportWidget = ViewportClient->GetEditorViewportWidget()->GetSceneViewport()->GetViewportWidget();
-	if (!ViewportWidget.IsValid())
+	TSharedPtr<SEditorViewport> EditorViewport = ViewportClient->GetEditorViewportWidget();
+	if (!EditorViewport)
+	{
+		return;
+	}
+
+	TSharedPtr<FSceneViewport> SceneViewport = EditorViewport->GetSceneViewport();
+	if (!SceneViewport)
 	{
 		return;
 	}
 
 	// set focus to viewport so that hotkeys are immediately detected
-	TSharedPtr<SWidget> ViewportContents = ViewportWidget.Pin()->GetContent();
-	FSlateApplication::Get().ForEachUser([&ViewportContents](FSlateUser& User)
+	if(const TSharedPtr<SViewport> ViewportWidget = SceneViewport->GetViewportWidget().Pin())
 	{
-		User.SetFocus(ViewportContents.ToSharedRef());
-	});
+		if (TSharedPtr<SWidget> ViewportContents = ViewportWidget->GetContent())
+		{
+			FSlateApplication::Get().ForEachUser([&ViewportContents](FSlateUser& User)
+			{
+				User.SetFocus(ViewportContents.ToSharedRef());
+			});
+		}
+	}
 }
 
 void UScriptableToolsEditorMode::OnBlueprintPreCompile(UBlueprint* Blueprint)
diff --git a/Engine/Plugins/Runtime/ScriptableToolsFramework/Source/ScriptableToolsFramework/Public/Utility/ScriptableToolContextObjects.h b/Engine/Plugins/Runtime/ScriptableToolsFramework/Source/ScriptableToolsFramework/Public/Utility/ScriptableToolContextObjects.h
--- a/Engine/Plugins/Runtime/ScriptableToolsFramework/Source/ScriptableToolsFramework/Public/Utility/ScriptableToolContextObjects.h
+++ b/Engine/Plugins/Runtime/ScriptableToolsFramework/Source/ScriptableToolsFramework/Public/Utility/ScriptableToolContextObjects.h
@@ -54,13 +54,19 @@ public:
 
 		CurrentOverlaidWidget = InOverlaidWidget.ToWeakPtr();
 
-		ReplaceOverlaidWidgetFunc(InOverlaidWidget);
+		if (ReplaceOverlaidWidgetFunc)
+		{
+			ReplaceOverlaidWidgetFunc(InOverlaidWidget);
+		}
 	}
 	void ClearOverlayWidget()
 	{
 		if (CurrentOverlaidWidget.IsValid())
 		{
-			ClearOverlaidWidgetFunc(CurrentOverlaidWidget.Pin().ToSharedRef());
+			if (ClearOverlaidWidgetFunc)
+			{
+				ClearOverlaidWidgetFunc(CurrentOverlaidWidget.Pin().ToSharedRef());
+			}
 			CurrentOverlaidWidget.Reset();
 		}
 	}
@@ -98,7 +104,10 @@ public:
 
 	void SetFocusInViewport()
 	{
-		SetFocusInViewportFunc();
+		if (SetFocusInViewportFunc)
+		{
+			SetFocusInViewportFunc();
+		}
 	}
 
 	virtual void Shutdown() override
