ScriptableTools: Restore behavior of hiding category toolbar if only a single palette is present.

diff --git a/Engine/Plugins/Editor/ScriptableToolsEditorMode/Source/ScriptableToolsEditorMode/Private/ScriptableToolsEditorModeToolkit.cpp b/Engine/Plugins/Editor/ScriptableToolsEditorMode/Source/ScriptableToolsEditorMode/Private/ScriptableToolsEditorModeToolkit.cpp
--- a/Engine/Plugins/Editor/ScriptableToolsEditorMode/Source/ScriptableToolsEditorMode/Private/ScriptableToolsEditorModeToolkit.cpp
+++ b/Engine/Plugins/Editor/ScriptableToolsEditorMode/Source/ScriptableToolsEditorMode/Private/ScriptableToolsEditorModeToolkit.cpp
@@ -1255,9 +1255,8 @@ void FScriptableToolsEditorModeToolkit::RebuildModeToolkitBuilderPalettes()
 		ToolkitBuilder->SetActivePaletteOnLoad(ActivePaletteOnLoad.Get());
 	}
 	
-	// Invoke UpdateWidget to update the CategoryToolbarVisibility on the ToolkitBuilder.
 	// This must be called prior to RefreshCategoryToolbarWidget since the visibility is not registered as a lambda.
-	ToolkitBuilder->UpdateWidget();
+	ToolkitBuilder->SetCategoryToolbarVisibility(PaletteCount > 1 ? EVisibility::Visible : EVisibility::Collapsed);
 	ToolkitBuilder->RefreshCategoryToolbarWidget(true);
 }
 
